Create InterchangeEditor plugin available by default. Project that need it at runtime will have to include Interchange plugin to there uproject.

diff --git a/Engine/Plugins/Editor/InterchangeEditor/InterchangeEditor.uplugin b/Engine/Plugins/Editor/InterchangeEditor/InterchangeEditor.uplugin
new file mode 100644
--- /dev/null
+++ b/Engine/Plugins/Editor/InterchangeEditor/InterchangeEditor.uplugin
@@ -0,0 +1,34 @@
+{
+	"FileVersion": 3,
+	"Version": 1,
+	"VersionName": "1.0",
+	"FriendlyName": "Interchange Editor",
+	"Description": "Interchange Editor plugin pull on Interchange plugin to allow us to force interchange plugin in editor.",
+	"Category": "Import Export",
+	"CreatedBy": "Epic Games, Inc.",
+	"CreatedByURL": "http://epicgames.com",
+	"DocsURL": "",
+	"MarketplaceURL": "",
+	"SupportURL": "",
+	"EnabledByDefault": true,
+	"CanContainContent": false,
+	"IsBetaVersion": false,
+	"Installed": false,
+	"Modules": [
+		{
+			"Name": "InterchangeEditor",
+			"Type": "Editor",
+			"LoadingPhase": "Default"
+		}
+	],
+	"Plugins": [
+		{
+			"Name": "Interchange",
+			"Enabled": true,
+			"WhitelistTargets": [
+				"Editor",
+				"Program"
+			]
+		}
+	]
+}
\ No newline at end of file
diff --git a/Engine/Plugins/Editor/InterchangeEditor/Source/InterchangeEditor/InterchangeEditor.Build.cs b/Engine/Plugins/Editor/InterchangeEditor/Source/InterchangeEditor/InterchangeEditor.Build.cs
new file mode 100644
--- /dev/null
+++ b/Engine/Plugins/Editor/InterchangeEditor/Source/InterchangeEditor/InterchangeEditor.Build.cs
@@ -0,0 +1,32 @@
+// Copyright Epic Games, Inc. All Rights Reserved.
+
+using System.IO;
+
+namespace UnrealBuildTool.Rules
+{
+	public class InterchangeEditor : ModuleRules
+	{
+		public InterchangeEditor(ReadOnlyTargetRules Target) : base(Target)
+		{
+            PublicDependencyModuleNames.AddRange(
+				new string[]
+				{
+                    "Core",
+					"CoreUObject"
+				}
+			);
+
+			PrivateDependencyModuleNames.AddRange(
+				new string[]
+				{
+					"InterchangeCore",
+					"InterchangeImportPlugin",
+					"InterchangeNodePlugin",
+					"InterchangeExportPlugin",
+					"InterchangeDispatcher",
+					"InterchangeFbxParser"
+				}
+			);
+		}
+    }
+}
diff --git a/Engine/Plugins/Editor/InterchangeEditor/Source/InterchangeEditor/Private/InterchangeEditorModule.cpp b/Engine/Plugins/Editor/InterchangeEditor/Source/InterchangeEditor/Private/InterchangeEditorModule.cpp
new file mode 100644
--- /dev/null
+++ b/Engine/Plugins/Editor/InterchangeEditor/Source/InterchangeEditor/Private/InterchangeEditorModule.cpp
@@ -0,0 +1,32 @@
+// Copyright Epic Games, Inc. All Rights Reserved.
+#include "InterchangeEditorModule.h"
+
+#include "InterchangeEditorLog.h"
+
+#include "Modules/ModuleInterface.h"
+#include "Modules/ModuleManager.h"
+
+
+
+#define LOCTEXT_NAMESPACE "InterchangeEditorModule"
+
+DEFINE_LOG_CATEGORY(LogInterchangeEditor);
+
+FInterchangeEditorModule& FInterchangeEditorModule::Get()
+{
+	return FModuleManager::LoadModuleChecked< FInterchangeEditorModule >(INTERCHANGEEDITOR_MODULE_NAME);
+}
+
+bool FInterchangeEditorModule::IsAvailable()
+{
+	return FModuleManager::Get().IsModuleLoaded(INTERCHANGEEDITOR_MODULE_NAME);
+}
+
+void FInterchangeEditorModule::StartupModule()
+{
+}
+
+IMPLEMENT_MODULE(FInterchangeEditorModule, InterchangeEditor);
+
+#undef LOCTEXT_NAMESPACE // "InterchangeEditorModule"
+
diff --git a/Engine/Plugins/Editor/InterchangeEditor/Source/InterchangeEditor/Public/InterchangeEditorLog.h b/Engine/Plugins/Editor/InterchangeEditor/Source/InterchangeEditor/Public/InterchangeEditorLog.h
new file mode 100644
--- /dev/null
+++ b/Engine/Plugins/Editor/InterchangeEditor/Source/InterchangeEditor/Public/InterchangeEditorLog.h
@@ -0,0 +1,7 @@
+// Copyright Epic Games, Inc. All Rights Reserved.
+#pragma once
+
+#include "Logging/LogMacros.h"
+
+DECLARE_LOG_CATEGORY_EXTERN(LogInterchangeEditor, Log, All)
+
diff --git a/Engine/Plugins/Editor/InterchangeEditor/Source/InterchangeEditor/Public/InterchangeEditorModule.h b/Engine/Plugins/Editor/InterchangeEditor/Source/InterchangeEditor/Public/InterchangeEditorModule.h
new file mode 100644
--- /dev/null
+++ b/Engine/Plugins/Editor/InterchangeEditor/Source/InterchangeEditor/Public/InterchangeEditorModule.h
@@ -0,0 +1,19 @@
+// Copyright Epic Games, Inc. All Rights Reserved.
+#pragma once
+
+#include "Modules/ModuleInterface.h"
+
+#define INTERCHANGEEDITOR_MODULE_NAME TEXT("InterchangeEditor")
+
+/**
+ * This module allow out of process interchange translator in case a third party SDK is not thread safe.
+ */
+class FInterchangeEditorModule : public IModuleInterface
+{
+public:
+	static FInterchangeEditorModule& Get();
+	static bool IsAvailable();
+
+private:
+	virtual void StartupModule() override;
+};
diff --git a/Engine/Plugins/Runtime/Interchange/Interchange.uplugin b/Engine/Plugins/Runtime/Interchange/Interchange.uplugin
--- a/Engine/Plugins/Runtime/Interchange/Interchange.uplugin
+++ b/Engine/Plugins/Runtime/Interchange/Interchange.uplugin
@@ -10,7 +10,7 @@
 	"DocsURL": "",
 	"MarketplaceURL": "",
 	"SupportURL": "",
-	"EnabledByDefault": true,
+	"EnabledByDefault": false,
 	"CanContainContent": false,
 	"IsBetaVersion": false,
 	"Installed": false,
@@ -40,7 +40,9 @@
 			"Type": "Runtime",
 			"LoadingPhase": "Default",
 			"WhitelistPlatforms": [
-				"Win64"
+				"Win64",
+				"Linux",
+				"Mac"
 			]
 		}
 	]
