Create a graph inspector pipeline to be able to debug interchange graph

diff --git a/Engine/Plugins/Editor/InterchangeEditor/InterchangeEditor.uplugin b/Engine/Plugins/Editor/InterchangeEditor/InterchangeEditor.uplugin
--- a/Engine/Plugins/Editor/InterchangeEditor/InterchangeEditor.uplugin
+++ b/Engine/Plugins/Editor/InterchangeEditor/InterchangeEditor.uplugin
@@ -19,6 +19,11 @@
 			"Name": "InterchangeEditor",
 			"Type": "Editor",
 			"LoadingPhase": "Default"
+		},
+		{
+			"Name": "InterchangeEditorPipelines",
+			"Type": "Editor",
+			"LoadingPhase": "Default"
 		}
 	],
 	"Plugins": [
diff --git a/Engine/Plugins/Editor/InterchangeEditor/Source/Pipelines/InterchangeEditorPipelines.Build.cs b/Engine/Plugins/Editor/InterchangeEditor/Source/Pipelines/InterchangeEditorPipelines.Build.cs
new file mode 100644
--- /dev/null
+++ b/Engine/Plugins/Editor/InterchangeEditor/Source/Pipelines/InterchangeEditorPipelines.Build.cs
@@ -0,0 +1,34 @@
+// Copyright Epic Games, Inc. All Rights Reserved.
+
+namespace UnrealBuildTool.Rules
+{
+	public class InterchangeEditorPipelines : ModuleRules
+	{
+		public InterchangeEditorPipelines(ReadOnlyTargetRules Target) : base(Target)
+		{
+			PublicDependencyModuleNames.AddRange(
+				new string[]
+				{
+					"Core",
+					"CoreUObject",
+					"Engine",
+					"InterchangeCore",
+					"InterchangeEngine",
+					"InterchangeNodes",
+				}
+			);
+
+			PrivateDependencyModuleNames.AddRange(
+				new string[]
+				{
+					"EditorStyle",
+					"InputCore",
+					"MainFrame",
+					"PropertyEditor",
+					"Slate",
+					"SlateCore",
+				}
+			);
+		}
+	}
+}
diff --git a/Engine/Plugins/Editor/InterchangeEditor/Source/Pipelines/Private/InterchangeEditorPipelineDetails.cpp b/Engine/Plugins/Editor/InterchangeEditor/Source/Pipelines/Private/InterchangeEditorPipelineDetails.cpp
new file mode 100644
--- /dev/null
+++ b/Engine/Plugins/Editor/InterchangeEditor/Source/Pipelines/Private/InterchangeEditorPipelineDetails.cpp
@@ -0,0 +1,647 @@
+// Copyright Epic Games, Inc. All Rights Reserved.
+
+#include "InterchangeEditorPipelineDetails.h"
+
+#include "DetailCategoryBuilder.h"
+#include "DetailLayoutBuilder.h"
+#include "DetailWidgetRow.h"
+#include "Editor.h"
+#include "IDetailGroup.h"
+#include "Nodes/InterchangeBaseNode.h"
+#include "PropertyHandle.h"
+#include "Widgets/DeclarativeSyntaxSupport.h"
+#include "Widgets/Images/SImage.h"
+#include "Widgets/Input/SCheckBox.h"
+#include "Widgets/Input/SNumericEntryBox.h"
+#include "Widgets/Layout/SBox.h"
+#include "Widgets/Text/STextBlock.h"
+#include "Widgets/SBoxPanel.h"
+
+
+#define LOCTEXT_NAMESPACE "InterchangeEditorPipelineDetails"
+
+FInterchangeBaseNodeDetailsCustomization::FInterchangeBaseNodeDetailsCustomization()
+{
+	InterchangeBaseNode = nullptr;
+	CachedDetailBuilder = nullptr;
+}
+
+FInterchangeBaseNodeDetailsCustomization::~FInterchangeBaseNodeDetailsCustomization()
+{
+}
+
+void FInterchangeBaseNodeDetailsCustomization::RefreshCustomDetail()
+{
+	if (CachedDetailBuilder)
+	{
+		CachedDetailBuilder->ForceRefreshDetails();
+	}
+}
+
+TSharedRef<IDetailCustomization> FInterchangeBaseNodeDetailsCustomization::MakeInstance()
+{
+	return MakeShareable(new FInterchangeBaseNodeDetailsCustomization);
+}
+
+void FInterchangeBaseNodeDetailsCustomization::CustomizeDetails( IDetailLayoutBuilder& DetailBuilder )
+{
+	CachedDetailBuilder = &DetailBuilder;
+	TArray<TWeakObjectPtr<UObject>> EditingObjects;
+	DetailBuilder.GetObjectsBeingCustomized(EditingObjects);
+	check(EditingObjects.Num() == 1);
+
+	InterchangeBaseNode = Cast<UInterchangeBaseNode>(EditingObjects[0].Get());
+
+	if (!ensure(InterchangeBaseNode))
+	{
+		return;
+	}
+
+	TArray<UE::Interchange::FAttributeKey> AttributeKeys;
+	InterchangeBaseNode->GetAttributeKeys(AttributeKeys);
+	IDetailCategoryBuilder& AttributeCategory = DetailBuilder.EditCategory("Attributes", FText::GetEmpty(), ECategoryPriority::Important);
+	for (UE::Interchange::FAttributeKey& AttributeKey : AttributeKeys)
+	{
+		AddAttributeRow(AttributeKey, AttributeCategory);
+	}
+}
+
+void FInterchangeBaseNodeDetailsCustomization::AddAttributeRow(UE::Interchange::FAttributeKey& AttributeKey, IDetailCategoryBuilder& AttributeCategory)
+{
+	if (!ensure(InterchangeBaseNode))
+	{
+		return;
+	}
+	const UE::Interchange::EAttributeTypes AttributeType = InterchangeBaseNode->GetAttributeType(AttributeKey);
+	switch (AttributeType)
+	{
+		case UE::Interchange::EAttributeTypes::Bool:
+		{
+			BuildBoolValueContent(AttributeCategory, AttributeKey);
+		}
+		break;
+
+		case UE::Interchange::EAttributeTypes::Double:
+		{
+			BuildNumberValueContent < double > (AttributeCategory, AttributeKey);
+		}
+		break;
+
+		case UE::Interchange::EAttributeTypes::Float:
+		{
+			BuildNumberValueContent < float >(AttributeCategory, AttributeKey);
+		}
+		break;
+
+		case UE::Interchange::EAttributeTypes::Int8:
+		{
+			BuildNumberValueContent < int8 >(AttributeCategory, AttributeKey);
+		}
+		break;
+
+		case UE::Interchange::EAttributeTypes::Int16:
+		{
+			BuildNumberValueContent < int16 >(AttributeCategory, AttributeKey);
+		}
+		break;
+
+		case UE::Interchange::EAttributeTypes::Int32:
+		{
+			BuildNumberValueContent < int32 >(AttributeCategory, AttributeKey);
+		}
+		break;
+
+		case UE::Interchange::EAttributeTypes::Int64:
+		{
+			BuildNumberValueContent < int64 >(AttributeCategory, AttributeKey);
+		}
+		break;
+
+		case UE::Interchange::EAttributeTypes::UInt8:
+		{
+			BuildNumberValueContent < uint8 >(AttributeCategory, AttributeKey);
+		}
+		break;
+
+		case UE::Interchange::EAttributeTypes::UInt16:
+		{
+			BuildNumberValueContent < uint16 >(AttributeCategory, AttributeKey);
+		}
+		break;
+
+		case UE::Interchange::EAttributeTypes::UInt32:
+		{
+			BuildNumberValueContent < uint32 >(AttributeCategory, AttributeKey);
+		}
+		break;
+
+		case UE::Interchange::EAttributeTypes::UInt64:
+		{
+			BuildNumberValueContent < uint64 >(AttributeCategory, AttributeKey);
+		}
+		break;
+
+		case UE::Interchange::EAttributeTypes::String:
+		{
+			BuildStringValueContent<FString>(AttributeCategory, AttributeKey);
+		}
+		break;
+
+		case UE::Interchange::EAttributeTypes::Name:
+		{
+			BuildStringValueContent<FName>(AttributeCategory, AttributeKey);
+		}
+		break;
+
+		case UE::Interchange::EAttributeTypes::Transform:
+		{
+			BuildTransformValueContent(AttributeCategory, AttributeKey);
+		}
+		break;
+
+		default:
+		{
+			FText AttributeName = FText::FromString(InterchangeBaseNode->GetKeyDisplayName(AttributeKey));
+			AttributeCategory.AddCustomRow(AttributeName)
+			.NameContent()
+			[
+				CreateNameWidget(AttributeKey)
+			]
+			.ValueContent()
+			[
+				SNew(SHorizontalBox)
+				+ SHorizontalBox::Slot()
+				.AutoWidth()
+				.HAlign(HAlign_Left)
+				.VAlign(VAlign_Center)
+				[
+					SNew(STextBlock)
+					.Text(LOCTEXT("UnsupportedCustomizationType", "Attribute Type Not Supported"))
+					.Font(IDetailLayoutBuilder::GetDetailFont())
+				]
+			];
+		}
+	}
+}
+
+void FInterchangeBaseNodeDetailsCustomization::BuildBoolValueContent(IDetailCategoryBuilder& AttributeCategory, UE::Interchange::FAttributeKey& AttributeKey)
+{
+	UE::Interchange::EAttributeTypes AttributeType = InterchangeBaseNode->GetAttributeType(AttributeKey);
+	check(AttributeType == UE::Interchange::EAttributeTypes::Bool);
+	UE::Interchange::FAttributeStorage::TAttributeHandle<bool> AttributeHandle = InterchangeBaseNode->GetAttributeHandle<bool>(AttributeKey);
+	if (!AttributeHandle.IsValid())
+	{
+		CreateInvalidHandleRow(AttributeCategory, AttributeKey);
+		return;
+	}
+	
+	const FText AttributeName = FText::FromString(InterchangeBaseNode->GetKeyDisplayName(AttributeKey));
+	FDetailWidgetRow& CustomRow = AttributeCategory.AddCustomRow(AttributeName)
+	.NameContent()
+	[
+		CreateNameWidget(AttributeKey)
+	]
+	.ValueContent()
+	[
+		SNew(SHorizontalBox)
+		+ SHorizontalBox::Slot()
+		.AutoWidth()
+		[
+			SNew(SBox)
+			[
+				SNew(SCheckBox)
+				.OnCheckStateChanged_Lambda([this, AttributeKey](ECheckBoxState CheckType)
+				{
+					const bool IsChecked = CheckType == ECheckBoxState::Checked;
+					UE::Interchange::FAttributeStorage::TAttributeHandle<bool> AttributeHandle = InterchangeBaseNode->GetAttributeHandle<bool>(AttributeKey);
+					if (AttributeHandle.IsValid())
+					{
+						AttributeHandle.Set(IsChecked);
+					}
+				})
+				.IsChecked_Lambda([this, AttributeKey]()
+				{
+					bool IsChecked = false;
+					UE::Interchange::FAttributeStorage::TAttributeHandle<bool> AttributeHandle = InterchangeBaseNode->GetAttributeHandle<bool>(AttributeKey);
+					if (AttributeHandle.IsValid())
+					{
+						AttributeHandle.Get(IsChecked);
+					}
+					return IsChecked ? ECheckBoxState::Checked : ECheckBoxState::Unchecked;
+				})
+			]
+		]
+	];
+}
+
+template<typename NumericType>
+void FInterchangeBaseNodeDetailsCustomization::BuildNumberValueContent(IDetailCategoryBuilder& AttributeCategory, UE::Interchange::FAttributeKey& AttributeKey)
+{
+	UE::Interchange::FAttributeStorage::TAttributeHandle<NumericType> AttributeHandle = InterchangeBaseNode->GetAttributeHandle<NumericType>(AttributeKey);
+	if (!AttributeHandle.IsValid())
+	{
+		CreateInvalidHandleRow(AttributeCategory, AttributeKey);
+		return;
+	}
+
+	auto GetValue = [](int32 ComponentIndex, UInterchangeBaseNode* BaseNode, UE::Interchange::FAttributeKey& Key)->NumericType
+	{
+		UE::Interchange::FAttributeStorage::TAttributeHandle<NumericType> AttributeHandle = BaseNode->GetAttributeHandle<NumericType>(Key);
+		//Prevent returning uninitialize value by setting it to 0
+		NumericType Value = static_cast<NumericType>(0);
+		if (AttributeHandle.IsValid())
+		{
+			AttributeHandle.Get(Value);
+		}
+		return Value;
+	};
+
+	auto SetValue = [](int32 ComponentIndex, UInterchangeBaseNode* BaseNode, NumericType Value, UE::Interchange::FAttributeKey& Key)
+	{
+		UE::Interchange::FAttributeStorage::TAttributeHandle<NumericType> AttributeHandle = BaseNode->GetAttributeHandle<NumericType>(Key);
+		if (AttributeHandle.IsValid())
+		{
+			AttributeHandle.Set(Value);
+		}
+    };
+
+	const FText AttributeName = FText::FromString(InterchangeBaseNode->GetKeyDisplayName(AttributeKey));
+	AttributeCategory.AddCustomRow(AttributeName)
+	.NameContent()
+	[
+		CreateNameWidget(AttributeKey)
+	]
+	.ValueContent()
+	[
+		SNew(SHorizontalBox)
+		+ SHorizontalBox::Slot()
+		.AutoWidth()
+		[
+			MakeNumericWidget<NumericType>(0, GetValue, SetValue, AttributeKey)
+		]
+	];
+}
+
+template<typename StringType>
+void FInterchangeBaseNodeDetailsCustomization::BuildStringValueContent(IDetailCategoryBuilder& AttributeCategory, UE::Interchange::FAttributeKey& AttributeKey)
+{
+	UE::Interchange::FAttributeStorage::TAttributeHandle<StringType> AttributeHandle = InterchangeBaseNode->GetAttributeHandle<StringType>(AttributeKey);
+	if (!AttributeHandle.IsValid())
+	{
+		CreateInvalidHandleRow(AttributeCategory, AttributeKey);
+		return;
+	}
+
+	const FText AttributeName = FText::FromString(InterchangeBaseNode->GetKeyDisplayName(AttributeKey));
+	FDetailWidgetRow& CustomRow = AttributeCategory.AddCustomRow(AttributeName)	
+    .NameContent()
+    [
+        CreateNameWidget(AttributeKey)
+    ]
+    .ValueContent()
+    [
+        SNew(SHorizontalBox)
+        + SHorizontalBox::Slot()
+        .AutoWidth()
+		.HAlign(HAlign_Left)
+		.VAlign(VAlign_Center)
+        [
+            SNew(STextBlock)
+            .Text_Lambda([this, AttributeKey]()->FText
+            {
+				UE::Interchange::EAttributeTypes AttributeType = InterchangeBaseNode->GetAttributeType(AttributeKey);
+				FText ReturnText;
+				if (AttributeType == UE::Interchange::EAttributeTypes::String)
+				{
+					UE::Interchange::FAttributeStorage::TAttributeHandle<FString> AttributeHandle = InterchangeBaseNode->GetAttributeHandle<FString>(AttributeKey);
+					FString Value;
+					if (AttributeHandle.IsValid())
+					{
+						AttributeHandle.Get(Value);
+					}
+					ReturnText = FText::FromString(Value);
+				}
+				else if (AttributeType == UE::Interchange::EAttributeTypes::Name)
+				{
+					UE::Interchange::FAttributeStorage::TAttributeHandle<FName> AttributeHandle = InterchangeBaseNode->GetAttributeHandle<FName>(AttributeKey);
+					FName Value;
+					if (AttributeHandle.IsValid())
+					{
+						AttributeHandle.Get(Value);
+					}
+					ReturnText = FText::FromName(Value);
+				}
+                return ReturnText;
+            })
+            .Font(IDetailLayoutBuilder::GetDetailFont())
+        ]
+    ];
+}
+
+void FInterchangeBaseNodeDetailsCustomization::BuildTransformValueContent(IDetailCategoryBuilder& AttributeCategory, UE::Interchange::FAttributeKey& AttributeKey)
+{
+	UE::Interchange::EAttributeTypes AttributeType = InterchangeBaseNode->GetAttributeType(AttributeKey);
+	check(AttributeType == UE::Interchange::EAttributeTypes::Transform);
+
+	{
+		const UE::Interchange::FAttributeStorage::TAttributeHandle<FTransform> AttributeHandle = InterchangeBaseNode->GetAttributeHandle<FTransform>(AttributeKey);
+		if (!AttributeHandle.IsValid())
+		{
+			CreateInvalidHandleRow(AttributeCategory, AttributeKey);
+			return;
+		}
+	}
+
+	const bool bAdvancedProperty = false;
+	const FString GroupName = InterchangeBaseNode->GetKeyDisplayName(AttributeKey);
+	IDetailGroup& Group = AttributeCategory.AddGroup(FName(*GroupName), FText::FromString(GroupName), bAdvancedProperty);
+	FDetailWidgetRow& GroupHeaderRow = Group.HeaderRow();
+	GroupHeaderRow.NameContent().Widget = SNew(SBox)
+	[
+		CreateNameWidget(AttributeKey)
+	];
+
+	auto GetRotationValue = [](UInterchangeBaseNode* BaseNode, UE::Interchange::FAttributeKey& Key)->FQuat
+	{
+		FTransform TransformValue;
+		const UE::Interchange::FAttributeStorage::TAttributeHandle<FTransform> AttributeHandle = BaseNode->GetAttributeHandle<FTransform>(Key);
+		if (AttributeHandle.IsValid())
+		{
+			AttributeHandle.Get(TransformValue);
+		}
+		return TransformValue.GetRotation();
+	};
+
+	auto SetRotationValue = [](UInterchangeBaseNode* BaseNode, UE::Interchange::FAttributeKey& Key, const FQuat& Value)
+	{
+		FTransform TransformValue;
+		UE::Interchange::FAttributeStorage::TAttributeHandle<FTransform> AttributeHandle = BaseNode->GetAttributeHandle<FTransform>(Key);
+		if (AttributeHandle.IsValid())
+		{
+			AttributeHandle.Get(TransformValue);
+			TransformValue.SetRotation(Value);
+			AttributeHandle.Set(TransformValue);
+		}
+	};
+
+	auto GetTranslationValue = [](UInterchangeBaseNode* BaseNode, UE::Interchange::FAttributeKey& Key)->FVector
+	{
+		FTransform TransformValue;
+		const UE::Interchange::FAttributeStorage::TAttributeHandle<FTransform> AttributeHandle = BaseNode->GetAttributeHandle<FTransform>(Key);
+		if (AttributeHandle.IsValid())
+		{
+			AttributeHandle.Get(TransformValue);
+		}
+		return TransformValue.GetTranslation();
+	};
+
+	auto SetTranslationValue = [](UInterchangeBaseNode* BaseNode, UE::Interchange::FAttributeKey& Key, const FVector& Value)
+	{
+		FTransform TransformValue;
+		UE::Interchange::FAttributeStorage::TAttributeHandle<FTransform> AttributeHandle = BaseNode->GetAttributeHandle<FTransform>(Key);
+		if (AttributeHandle.IsValid())
+		{
+			AttributeHandle.Get(TransformValue);
+			TransformValue.SetTranslation(Value);
+			AttributeHandle.Set(TransformValue);
+		}
+	};
+
+	auto GetScale3DValue = [](UInterchangeBaseNode* BaseNode, UE::Interchange::FAttributeKey& Key)->FVector
+	{
+		FTransform TransformValue;
+		const UE::Interchange::FAttributeStorage::TAttributeHandle<FTransform> AttributeHandle = BaseNode->GetAttributeHandle<FTransform>(Key);
+		if (AttributeHandle.IsValid())
+		{
+			AttributeHandle.Get(TransformValue);
+		}
+		return TransformValue.GetScale3D();
+	};
+
+	auto SetScale3DValue = [](UInterchangeBaseNode* BaseNode, UE::Interchange::FAttributeKey& Key, const FVector& Value)
+	{
+		FTransform TransformValue;
+		UE::Interchange::FAttributeStorage::TAttributeHandle<FTransform> AttributeHandle = BaseNode->GetAttributeHandle<FTransform>(Key);
+		if (AttributeHandle.IsValid())
+		{
+			AttributeHandle.Get(TransformValue);
+			TransformValue.SetScale3D(Value);
+			AttributeHandle.Set(TransformValue);
+		}
+	};
+
+	const FString TranslationName = TEXT("Translation");
+	Group.AddWidgetRow()
+    .NameContent()
+    [
+        CreateSimpleNameWidget(TranslationName)
+    ]
+    .ValueContent()
+    [
+        CreateVectorWidget(GetTranslationValue, SetTranslationValue, AttributeKey)
+    ];
+
+	const FString RotationName = TEXT("Rotation");
+	Group.AddWidgetRow()
+	.NameContent()
+	[
+		CreateSimpleNameWidget(RotationName)
+	]
+	.ValueContent()
+	[
+		CreateQuaternionWidget(GetRotationValue, SetRotationValue, AttributeKey)
+	];
+
+	const FString Scale3DName = TEXT("Scale3D");
+	Group.AddWidgetRow()
+    .NameContent()
+    [
+        CreateSimpleNameWidget(Scale3DName)
+    ]
+    .ValueContent()
+    [
+        CreateVectorWidget(GetScale3DValue, SetScale3DValue, AttributeKey)
+    ];
+}
+
+void FInterchangeBaseNodeDetailsCustomization::CreateInvalidHandleRow(IDetailCategoryBuilder& AttributeCategory, UE::Interchange::FAttributeKey& AttributeKey) const
+{
+	const FString InvalidAttributeHandle = TEXT("Invalid Attribute Handle!");
+	const FText AttributeName = FText::FromString(InterchangeBaseNode->GetKeyDisplayName(AttributeKey));
+	AttributeCategory.AddCustomRow(AttributeName)
+	.NameContent()
+	[
+		CreateNameWidget(AttributeKey)
+	]
+	.ValueContent()
+	[
+		CreateSimpleNameWidget(InvalidAttributeHandle)
+	];
+}
+
+TSharedRef<SWidget> FInterchangeBaseNodeDetailsCustomization::CreateNameWidget(UE::Interchange::FAttributeKey& AttributeKey) const
+{
+	const UE::Interchange::EAttributeTypes AttributeType = InterchangeBaseNode->GetAttributeType(AttributeKey);
+	const FText AttributeName = FText::FromString(InterchangeBaseNode->GetKeyDisplayName(AttributeKey));
+	const FString AttributeTooltipString = TEXT("Attribute Type: ") + UE::Interchange::AttributeTypeToString(AttributeType);
+	const FText AttributeTooltip = FText::FromString(AttributeTooltipString);
+	return SNew(STextBlock)
+		.Text(AttributeName)
+		.Font(IDetailLayoutBuilder::GetDetailFont())
+		.ToolTipText(AttributeTooltip);
+		
+}
+
+TSharedRef<SWidget> FInterchangeBaseNodeDetailsCustomization::CreateSimpleNameWidget(const FString& SimpleName) const
+{
+	const FText SimpleNameText = FText::FromString(SimpleName);
+	return SNew(STextBlock)
+		.Text(SimpleNameText)
+		.Font(IDetailLayoutBuilder::GetDetailFont());
+}
+
+template<typename FunctorGet, typename FunctorSet>
+TSharedRef<SWidget> FInterchangeBaseNodeDetailsCustomization::CreateVectorWidget(FunctorGet GetValue, FunctorSet SetValue, UE::Interchange::FAttributeKey& AttributeKey)
+{
+	auto GetComponentValue = [&GetValue](int32 ComponentIndex, UInterchangeBaseNode* BaseNode, UE::Interchange::FAttributeKey& Key)->float
+	{
+		FVector Value = GetValue(BaseNode, Key);
+		return Value[ComponentIndex];
+	};
+
+	auto SetComponentValue = [&GetValue, &SetValue](int32 ComponentIndex, UInterchangeBaseNode* BaseNode, float ComponentValue, UE::Interchange::FAttributeKey& Key)
+	{
+		FVector Value = GetValue(BaseNode, Key);
+		Value[ComponentIndex] = ComponentValue;
+		SetValue(BaseNode, Key, Value);
+    };
+	
+	//Create a horizontal layout with the 3 floats components
+	return SNew(SHorizontalBox)
+        + SHorizontalBox::Slot()
+        .AutoWidth()
+        [
+			MakeNumericWidget<float>(0, GetComponentValue, SetComponentValue, AttributeKey)
+        ]
+		+ SHorizontalBox::Slot()
+	    .AutoWidth()
+	    [
+	        MakeNumericWidget<float>(1, GetComponentValue, SetComponentValue, AttributeKey)
+	    ]
+		+ SHorizontalBox::Slot()
+	    .AutoWidth()
+	    [
+	        MakeNumericWidget<float>(2, GetComponentValue, SetComponentValue, AttributeKey)
+	    ];
+}
+
+template<typename FunctorGet, typename FunctorSet>
+TSharedRef<SWidget> FInterchangeBaseNodeDetailsCustomization::CreateQuaternionWidget(FunctorGet GetValue, FunctorSet SetValue, UE::Interchange::FAttributeKey& AttributeKey)
+{
+	auto GetComponentValue = [&GetValue](int32 ComponentIndex, UInterchangeBaseNode* BaseNode, UE::Interchange::FAttributeKey& Key)->float
+	{
+		FQuat Value = GetValue(BaseNode, Key);
+		if(ComponentIndex == 0)
+		{
+			return Value.X;
+		}
+		else if(ComponentIndex == 1)
+		{
+			return Value.Y;
+		}
+		else if(ComponentIndex == 2)
+		{
+			return Value.Z;
+		}
+		else if(ComponentIndex == 3)
+		{
+			return Value.W;
+		}
+
+		//Ensure
+		ensure(ComponentIndex >= 0 && ComponentIndex < 4);
+		return 0.0f;	
+	};
+
+	auto SetComponentValue = [&GetValue, &SetValue](int32 ComponentIndex, UInterchangeBaseNode* BaseNode, float ComponentValue, UE::Interchange::FAttributeKey& Key)
+	{
+		FQuat Value = GetValue(BaseNode, Key);
+		if(ComponentIndex == 0)
+		{
+			Value.X = ComponentValue;
+		}
+		else if(ComponentIndex == 1)
+		{
+			Value.Y = ComponentValue;
+		}
+		else if(ComponentIndex == 2)
+		{
+			Value.Z = ComponentValue;
+		}
+		else if(ComponentIndex == 3)
+		{
+			Value.W = ComponentValue;
+		}
+		//Ensure
+		if(ensure(ComponentIndex >= 0 && ComponentIndex < 4))
+		{
+			SetValue(BaseNode, Key, Value);
+		}
+	};
+	
+	//Create a horizontal layout with the 4 floats components
+	return SNew(SHorizontalBox)
+        + SHorizontalBox::Slot()
+        .AutoWidth()
+        [
+            MakeNumericWidget<float>(0, GetComponentValue, SetComponentValue, AttributeKey)
+        ]
+        + SHorizontalBox::Slot()
+        .AutoWidth()
+        [
+            MakeNumericWidget<float>(1, GetComponentValue, SetComponentValue, AttributeKey)
+        ]
+        + SHorizontalBox::Slot()
+        .AutoWidth()
+        [
+            MakeNumericWidget<float>(2, GetComponentValue, SetComponentValue, AttributeKey)
+        ]
+		+ SHorizontalBox::Slot()
+	    .AutoWidth()
+	    [
+	        MakeNumericWidget<float>(3, GetComponentValue, SetComponentValue, AttributeKey)
+	    ];
+}
+
+template<typename NumericType, typename FunctorGet, typename FunctorSet>
+TSharedRef<SWidget> FInterchangeBaseNodeDetailsCustomization::MakeNumericWidget(int32 ComponentIndex, FunctorGet GetValue, FunctorSet SetValue, UE::Interchange::FAttributeKey& AttributeKey)
+{
+	//The 3 Lambda functions need to call other lambda function having AttributeKey reference (&AttributeKey).
+	//The per value capture prevent us the AttributeKey by reference 
+	//We copy the key only one time here and pass a reference of this copy
+
+	auto SetValueCommittedLambda = [this, &SetValue, ComponentIndex, AttributeKey](const NumericType Value, ETextCommit::Type CommitType)
+	{
+		UE::Interchange::FAttributeKey Key = AttributeKey;
+		SetValue(ComponentIndex, InterchangeBaseNode, Value, Key);
+	};
+	auto SetValueChangedLambda = [this, &SetValue, ComponentIndex, AttributeKey](const NumericType Value)
+	{
+		UE::Interchange::FAttributeKey Key = AttributeKey;
+		SetValue(ComponentIndex, InterchangeBaseNode, Value, Key);
+	};
+	
+	return
+		SNew(SNumericEntryBox<NumericType>)
+			.EditableTextBoxStyle(&FCoreStyle::Get().GetWidgetStyle<FEditableTextBoxStyle>("NormalEditableTextBox"))
+			.Value_Lambda([this, &GetValue, ComponentIndex, AttributeKey]()
+			{
+				UE::Interchange::FAttributeKey Key = AttributeKey;
+				return GetValue(ComponentIndex, InterchangeBaseNode, Key);
+			})
+			.Font(IDetailLayoutBuilder::GetDetailFont())
+			.OnValueCommitted_Lambda(SetValueCommittedLambda)
+			.OnValueChanged_Lambda(SetValueChangedLambda)
+			.AllowSpin(false);
+}
+#undef LOCTEXT_NAMESPACE
diff --git a/Engine/Plugins/Editor/InterchangeEditor/Source/Pipelines/Private/InterchangeEditorPipelineDetails.h b/Engine/Plugins/Editor/InterchangeEditor/Source/Pipelines/Private/InterchangeEditorPipelineDetails.h
new file mode 100644
--- /dev/null
+++ b/Engine/Plugins/Editor/InterchangeEditor/Source/Pipelines/Private/InterchangeEditorPipelineDetails.h
@@ -0,0 +1,100 @@
+// Copyright Epic Games, Inc. All Rights Reserved.
+
+#pragma once
+
+#include "CoreMinimal.h"
+#include "UObject/WeakObjectPtr.h"
+#include "Types/SlateEnums.h"
+#include "IDetailCustomization.h"
+#include "Input/Reply.h"
+
+class IDetailCategoryBuilder;
+class IDetailLayoutBuilder;
+class IDetailPropertyRow;
+class IPropertyHandle;
+class UInterchangeBaseNode;
+
+namespace UE
+{
+	namespace Interchange
+	{
+		struct FAttributeKey;
+	}
+}
+
+
+class FInterchangeBaseNodeDetailsCustomization : public IDetailCustomization
+{
+public:
+	~FInterchangeBaseNodeDetailsCustomization();
+
+	/** Makes a new instance of this detail layout class for a specific detail view requesting it */
+	static TSharedRef<IDetailCustomization> MakeInstance();
+
+	/** IDetailCustomization interface */
+	virtual void CustomizeDetails( IDetailLayoutBuilder& DetailBuilder ) override;
+	/** End IDetailCustomization interface */
+
+private:
+	void RefreshCustomDetail();
+	void AddAttributeRow(UE::Interchange::FAttributeKey& AttributeKey, IDetailCategoryBuilder& AttributeCategory);
+
+	/*******************************************************************
+	 * AttributeType Build functions Begin
+	 */
+	
+	 //Bool: Add a row to the category containing a check box (read/write)
+	void BuildBoolValueContent(IDetailCategoryBuilder& AttributeCategory, UE::Interchange::FAttributeKey& AttributeKey);
+	
+	//Number (float, double, signed/unsigned integer): Add one row to the category containing a numeric field (read/write)
+	template<typename NumericType>
+	void BuildNumberValueContent(IDetailCategoryBuilder& AttributeCategory, UE::Interchange::FAttributeKey& AttributeKey);
+
+	//String (FString, FName): Add a row to the category containing text block (read only)
+	template<typename StringType>
+	void BuildStringValueContent(IDetailCategoryBuilder& AttributeCategory, UE::Interchange::FAttributeKey& AttributeKey);
+
+	//FTransform: Add a group to the category that contain the Translation(FVector), Rotation(FQuat) and scale3D(FVector)
+	//Any numeric field are read/write
+	void BuildTransformValueContent(IDetailCategoryBuilder& AttributeCategory, UE::Interchange::FAttributeKey& AttributeKey);
+
+	/**
+	 * AttributeType Build functions End
+	 ******************************************************************/
+
+	/*
+	 * Add a row to the specified category which let the user that a attribute handle is invalid 
+	 */
+	void CreateInvalidHandleRow(IDetailCategoryBuilder& AttributeCategory, UE::Interchange::FAttributeKey& AttributeKey) const;
+
+
+	/*******************************************************************
+	 * Basic type widgets creation Begin
+	 */
+
+	TSharedRef<SWidget> CreateNameWidget(UE::Interchange::FAttributeKey& AttributeKey) const;
+
+	TSharedRef<SWidget> CreateSimpleNameWidget(const FString& SimpleName) const;
+
+	//FVector Create a 3 digit widget (X, Y, Z)
+	template<typename FunctorGet, typename FunctorSet>
+	TSharedRef<SWidget> CreateVectorWidget(FunctorGet GetValue, FunctorSet SetValue, UE::Interchange::FAttributeKey& AttributeKey);
+
+	//FQuat Create a 4 digit widget (X, Y, Z, W)
+	template<typename FunctorGet, typename FunctorSet>
+    TSharedRef<SWidget> CreateQuaternionWidget(FunctorGet GetValue, FunctorSet SetValue, UE::Interchange::FAttributeKey& AttributeKey);
+
+	//Create one numeric widget, support (float, double, signed/unsigned integer)
+	template<typename NumericType, typename FunctorGet, typename FunctorSet>
+	TSharedRef<SWidget> MakeNumericWidget(int32 ComponentIndex, FunctorGet GetValue, FunctorSet SetValue, UE::Interchange::FAttributeKey& AttributeKey);
+
+	/*
+	 * Basic type widgets creation End
+	 *******************************************************************/
+
+	/** Use MakeInstance to create an instance of this class */
+	FInterchangeBaseNodeDetailsCustomization();
+
+	UInterchangeBaseNode* InterchangeBaseNode;
+	IDetailLayoutBuilder* CachedDetailBuilder;	// The detail builder for this customisation
+};
diff --git a/Engine/Plugins/Editor/InterchangeEditor/Source/Pipelines/Private/InterchangeEditorPipelinesModule.cpp b/Engine/Plugins/Editor/InterchangeEditor/Source/Pipelines/Private/InterchangeEditorPipelinesModule.cpp
new file mode 100644
--- /dev/null
+++ b/Engine/Plugins/Editor/InterchangeEditor/Source/Pipelines/Private/InterchangeEditorPipelinesModule.cpp
@@ -0,0 +1,63 @@
+// Copyright Epic Games, Inc. All Rights Reserved.
+
+#include "InterchangeEditorPipelinesModule.h"
+
+#include "CoreMinimal.h"
+#include "Engine/Engine.h"
+#include "InterchangeEditorPipelineDetails.h"
+#include "InterchangeGraphInspectorPipeline.h"
+#include "InterchangeManager.h"
+#include "Misc/CoreDelegates.h"
+#include "Modules/ModuleManager.h"
+#include "Nodes/InterchangeBaseNode.h"
+#include "PropertyEditorModule.h"
+
+
+class FInterchangeEditorPipelinesModule : public IInterchangeEditorPipelinesModule
+{
+	/** IModuleInterface implementation */
+	virtual void StartupModule() override;
+	virtual void ShutdownModule() override;
+};
+
+IMPLEMENT_MODULE(FInterchangeEditorPipelinesModule, InterchangeEditorPipelines)
+
+void FInterchangeEditorPipelinesModule::StartupModule()
+{
+	auto RegisterItems = []()
+	{
+		UInterchangeManager& InterchangeManager = UInterchangeManager::GetInterchangeManager();
+	};
+
+	if (GEngine)
+	{
+		RegisterItems();
+	}
+	else
+	{
+		FCoreDelegates::OnPostEngineInit.AddLambda(RegisterItems);
+	}
+
+	ClassesToUnregisterOnShutdown.Reset();
+	// Register details customizations for animation controller nodes
+	FPropertyEditorModule& PropertyEditorModule = FModuleManager::LoadModuleChecked<FPropertyEditorModule>("PropertyEditor");
+
+	ClassesToUnregisterOnShutdown.Add(UInterchangeBaseNode::StaticClass()->GetFName());
+	PropertyEditorModule.RegisterCustomClassLayout(ClassesToUnregisterOnShutdown.Last(), FOnGetDetailCustomizationInstance::CreateStatic(&FInterchangeBaseNodeDetailsCustomization::MakeInstance));
+}
+
+
+void FInterchangeEditorPipelinesModule::ShutdownModule()
+{
+	FPropertyEditorModule* PropertyEditorModule = FModuleManager::GetModulePtr<FPropertyEditorModule>("PropertyEditor");
+	if (PropertyEditorModule)
+	{
+		for (FName ClassName : ClassesToUnregisterOnShutdown)
+		{
+			PropertyEditorModule->UnregisterCustomClassLayout(ClassName);
+		}
+	}
+}
+
+
+
diff --git a/Engine/Plugins/Editor/InterchangeEditor/Source/Pipelines/Private/InterchangeGraphInspectorPipeline.cpp b/Engine/Plugins/Editor/InterchangeEditor/Source/Pipelines/Private/InterchangeGraphInspectorPipeline.cpp
new file mode 100644
--- /dev/null
+++ b/Engine/Plugins/Editor/InterchangeEditor/Source/Pipelines/Private/InterchangeGraphInspectorPipeline.cpp
@@ -0,0 +1,39 @@
+// Copyright Epic Games, Inc. All Rights Reserved. 
+#include "InterchangeGraphInspectorPipeline.h"
+
+#include "CoreMinimal.h"
+#include "InterchangeSkeletalMeshLodDataNode.h"
+#include "InterchangeSkeletalMeshNode.h"
+#include "InterchangeSkeletonNode.h"
+#include "InterchangeSourceData.h"
+#include "Interfaces/IMainFrameModule.h"
+#include "Misc/Paths.h"
+#include "Nodes/InterchangeBaseNodeContainer.h"
+#include "SInterchangeGraphInspectorWindow.h"
+#include "UObject/Object.h"
+#include "UObject/ObjectMacros.h"
+
+bool UInterchangeGraphInspectorPipeline::ExecutePreImportPipeline(UInterchangeBaseNodeContainer* BaseNodeContainer, const TArray<UInterchangeSourceData*>& SourceDatas)
+{
+	//Create and show the graph inspector UI dialog
+	TSharedPtr<SWindow> ParentWindow;
+	if (FModuleManager::Get().IsModuleLoaded("MainFrame"))
+	{
+		IMainFrameModule& MainFrame = FModuleManager::LoadModuleChecked<IMainFrameModule>("MainFrame");
+		ParentWindow = MainFrame.GetParentWindow();
+	}
+	TSharedRef<SWindow> Window = SNew(SWindow)
+		.ClientSize(FVector2D(800.f, 650.f))
+		.Title(NSLOCTEXT("Interchange", "GraphInspectorTitle", "Interchange node graph inspector"));
+	TSharedPtr<SInterchangeGraphInspectorWindow> InterchangeGraphInspectorWindow;
+
+	Window->SetContent
+	(
+		SAssignNew(InterchangeGraphInspectorWindow, SInterchangeGraphInspectorWindow)
+		.InterchangeBaseNodeContainer(BaseNodeContainer)
+		.OwnerWindow(Window)
+	);
+
+	FSlateApplication::Get().AddModalWindow(Window, ParentWindow, false);
+	return true;
+}
diff --git a/Engine/Plugins/Editor/InterchangeEditor/Source/Pipelines/Private/SInterchangeGraphInspectorWindow.cpp b/Engine/Plugins/Editor/InterchangeEditor/Source/Pipelines/Private/SInterchangeGraphInspectorWindow.cpp
new file mode 100644
--- /dev/null
+++ b/Engine/Plugins/Editor/InterchangeEditor/Source/Pipelines/Private/SInterchangeGraphInspectorWindow.cpp
@@ -0,0 +1,484 @@
+// Copyright Epic Games, Inc. All Rights Reserved.
+#include "SInterchangeGraphInspectorWindow.h"
+
+#include "EditorStyleSet.h"
+#include "IDetailsView.h"
+#include "IDocumentation.h"
+#include "Modules/ModuleManager.h"
+#include "Nodes/InterchangeBaseNode.h"
+#include "Nodes/InterchangeBaseNodeContainer.h"
+#include "PropertyEditorModule.h"
+#include "Styling/SlateIconFinder.h"
+#include "Widgets/Images/SImage.h"
+#include "Widgets/Input/SButton.h"
+#include "Widgets/Input/SCheckBox.h"
+#include "Widgets/Input/SEditableText.h"
+#include "Widgets/Input/STextComboBox.h"
+#include "Widgets/Layout/SSplitter.h"
+#include "Widgets/Layout/SUniformGridPanel.h"
+
+#define LOCTEXT_NAMESPACE "InterchangeGraphInspector"
+
+/************************************************************************/
+/* SInterchangeGraphInspectorTreeView Implementation                    */
+/************************************************************************/
+
+SInterchangeGraphInspectorTreeView::~SInterchangeGraphInspectorTreeView()
+{
+
+}
+
+void SInterchangeGraphInspectorTreeView::Construct(const FArguments& InArgs)
+{
+	InterchangeBaseNodeContainer = InArgs._InterchangeBaseNodeContainer;
+	OnSelectionChangedDelegate = InArgs._OnSelectionChangedDelegate;
+	//Build the FbxNodeInfoPtr tree data
+	check(InterchangeBaseNodeContainer != nullptr);
+
+	TArray<FString> Roots;
+	InterchangeBaseNodeContainer->GetRoots(Roots);
+	for (FString RootID : Roots)
+	{
+		UInterchangeBaseNode* RootNode = InterchangeBaseNodeContainer->GetNode(RootID);
+
+		RootNodeArray.Add(RootNode);
+	}
+
+	STreeView::Construct
+	(
+		STreeView::FArguments()
+		.TreeItemsSource(&RootNodeArray)
+		.SelectionMode(ESelectionMode::Multi)
+		.OnGenerateRow(this, &SInterchangeGraphInspectorTreeView::OnGenerateRowGraphInspectorTreeView)
+		.OnGetChildren(this, &SInterchangeGraphInspectorTreeView::OnGetChildrenGraphInspectorTreeView)
+		.OnContextMenuOpening(this, &SInterchangeGraphInspectorTreeView::OnOpenContextMenu)
+		.OnSelectionChanged(this, &SInterchangeGraphInspectorTreeView::OnTreeViewSelectionChanged)
+	);
+}
+
+/** The item used for visualizing the class in the tree. */
+class SInterchangeGraphInspectorTreeViewItem : public STableRow< UInterchangeBaseNode* >
+{
+public:
+
+	SLATE_BEGIN_ARGS(SInterchangeGraphInspectorTreeViewItem)
+		: _InterchangeNode(nullptr)
+		, _InterchangeBaseNodeContainer(nullptr)
+	{}
+
+	/** The item content. */
+	SLATE_ARGUMENT(UInterchangeBaseNode*, InterchangeNode)
+	SLATE_ARGUMENT(UInterchangeBaseNodeContainer*, InterchangeBaseNodeContainer)
+	SLATE_END_ARGS()
+
+	/**
+	* Construct the widget
+	*
+	* @param InArgs   A declaration from which to construct the widget
+	*/
+	void Construct(const FArguments& InArgs, const TSharedRef<STableViewBase>& InOwnerTableView)
+	{
+		InterchangeNode = InArgs._InterchangeNode;
+		InterchangeBaseNodeContainer = InArgs._InterchangeBaseNodeContainer;
+
+		//This is suppose to always be valid
+		check(InterchangeNode);
+		check(InterchangeBaseNodeContainer);
+
+		UClass* IconClass = AActor::StaticClass();
+		if (InterchangeNode->GetAssetClass())
+		{
+			IconClass = InterchangeNode->GetAssetClass();
+		}
+		const FSlateBrush* ClassIcon = FSlateIconFinder::FindIconBrushForClass(IconClass);
+
+		//Prepare the tooltip
+		FString Tooltip = InterchangeNode->GetDisplayLabel();
+		if (InterchangeNode->GetAssetClass())
+		{
+			Tooltip += TEXT(" [") + InterchangeNode->GetAssetClass()->GetName() + TEXT("]");
+		}
+
+		this->ChildSlot
+		[
+			SNew(SHorizontalBox)
+			+ SHorizontalBox::Slot()
+			.AutoWidth()
+			.Padding(2.0f, 0.0f, 2.0f, 0.0f)
+			[
+				SNew(SCheckBox)
+				.OnCheckStateChanged(this, &SInterchangeGraphInspectorTreeViewItem::OnItemCheckChanged)
+				.IsChecked(this, &SInterchangeGraphInspectorTreeViewItem::IsItemChecked)
+			]
+			+ SHorizontalBox::Slot()
+			.AutoWidth()
+			[
+				SNew(SExpanderArrow, SharedThis(this))
+			]
+			+ SHorizontalBox::Slot()
+			.AutoWidth()
+			.Padding(0.0f, 2.0f, 6.0f, 2.0f)
+			[
+				SNew(SImage)
+				.Image(ClassIcon)
+				.Visibility(ClassIcon != FEditorStyle::GetDefaultBrush() ? EVisibility::Visible : EVisibility::Collapsed)
+			]
+			+ SHorizontalBox::Slot()
+			.FillWidth(1.0f)
+			.Padding(0.0f, 3.0f, 6.0f, 3.0f)
+			.VAlign(VAlign_Center)
+			[
+				SNew(STextBlock)
+				.Text(FText::FromString(InterchangeNode->GetDisplayLabel()))
+				.ToolTipText(FText::FromString(Tooltip))
+			]
+
+		];
+
+		STableRow< UInterchangeBaseNode* >::ConstructInternal(
+			STableRow::FArguments()
+			.ShowSelection(true),
+			InOwnerTableView
+		);
+	}
+
+private:
+
+	void RecursiveSetImport(UInterchangeBaseNode* Node, bool bState)
+	{
+		TArray<FString> Childrens = InterchangeBaseNodeContainer->GetNodeChildrenUIDs(Node->GetUniqueID());
+		for (const FString& ChildID : Childrens)
+		{
+			UInterchangeBaseNode* ChildNode = InterchangeBaseNodeContainer->GetNode(ChildID);
+			if (!ChildNode)
+				continue;
+			ChildNode->SetEnabled(bState);
+			RecursiveSetImport(ChildNode, bState);
+		}
+	}
+
+	void OnItemCheckChanged(ECheckBoxState CheckType)
+	{
+		if (!InterchangeNode)
+			return;
+		const bool ImportState = CheckType == ECheckBoxState::Checked;
+		InterchangeNode->SetEnabled(ImportState);
+		RecursiveSetImport(InterchangeNode, ImportState);
+	}
+
+	ECheckBoxState IsItemChecked() const
+	{
+		return InterchangeNode->IsEnabled() ? ECheckBoxState::Checked : ECheckBoxState::Unchecked;
+	}
+
+	/** The node to build the tree view row from. */
+	UInterchangeBaseNode* InterchangeNode = nullptr;
+	UInterchangeBaseNodeContainer* InterchangeBaseNodeContainer = nullptr;
+};
+
+TSharedRef< ITableRow > SInterchangeGraphInspectorTreeView::OnGenerateRowGraphInspectorTreeView(UInterchangeBaseNode* Item, const TSharedRef< STableViewBase >& OwnerTable)
+{
+	TSharedRef< SInterchangeGraphInspectorTreeViewItem > ReturnRow = SNew(SInterchangeGraphInspectorTreeViewItem, OwnerTable)
+		.InterchangeNode(Item)
+		.InterchangeBaseNodeContainer(InterchangeBaseNodeContainer);
+	return ReturnRow;
+}
+void SInterchangeGraphInspectorTreeView::OnGetChildrenGraphInspectorTreeView(UInterchangeBaseNode* InParent, TArray< UInterchangeBaseNode* >& OutChildren)
+{
+	TArray<FString> Childrens = InterchangeBaseNodeContainer->GetNodeChildrenUIDs(InParent->GetUniqueID());
+	for (const FString& ChildID : Childrens)
+	{
+		UInterchangeBaseNode* ChildNode = InterchangeBaseNodeContainer->GetNode(ChildID);
+		if (!ChildNode)
+			continue;
+		OutChildren.Add(ChildNode);
+	}
+}
+
+void SInterchangeGraphInspectorTreeView::RecursiveSetImport(UInterchangeBaseNode* Node, bool bState)
+{
+	TArray<FString> Childrens = InterchangeBaseNodeContainer->GetNodeChildrenUIDs(Node->GetUniqueID());
+	for (const FString& ChildID : Childrens)
+	{
+		UInterchangeBaseNode* ChildNode = InterchangeBaseNodeContainer->GetNode(ChildID);
+		if (!ChildNode)
+			continue;
+		ChildNode->SetEnabled(bState);
+		RecursiveSetImport(ChildNode, bState);
+	}
+}
+
+void SInterchangeGraphInspectorTreeView::OnToggleSelectAll(ECheckBoxState CheckType)
+{
+	for (UInterchangeBaseNode* Node : RootNodeArray)
+	{
+		if (!ensure(Node))
+		{
+			continue;
+		}
+		RecursiveSetImport(Node, CheckType == ECheckBoxState::Checked);
+	}
+}
+
+void SInterchangeGraphInspectorTreeView::RecursiveSetExpand(UInterchangeBaseNode* Node, bool ExpandState)
+{
+	SetItemExpansion(Node, ExpandState);
+	TArray<FString> Childrens = InterchangeBaseNodeContainer->GetNodeChildrenUIDs(Node->GetUniqueID());
+	for (const FString& ChildID : Childrens)
+	{
+		UInterchangeBaseNode* ChildNode = InterchangeBaseNodeContainer->GetNode(ChildID);
+		if (!ChildNode)
+			continue;
+		RecursiveSetExpand(ChildNode, ExpandState);
+	}
+}
+
+FReply SInterchangeGraphInspectorTreeView::OnExpandAll()
+{
+	for (UInterchangeBaseNode* Node : RootNodeArray)
+	{
+		if (!ensure(Node))
+		{
+			continue;
+		}
+		RecursiveSetExpand(Node, true);
+	}
+	return FReply::Handled();
+}
+
+FReply SInterchangeGraphInspectorTreeView::OnCollapseAll()
+{
+	for (UInterchangeBaseNode* Node : RootNodeArray)
+	{
+		if (!ensure(Node))
+		{
+			continue;
+		}
+		RecursiveSetExpand(Node, false);
+	}
+	return FReply::Handled();
+}
+
+TSharedPtr<SWidget> SInterchangeGraphInspectorTreeView::OnOpenContextMenu()
+{
+	// Build up the menu for a selection
+	const bool bCloseAfterSelection = true;
+	FMenuBuilder MenuBuilder(bCloseAfterSelection, TSharedPtr<FUICommandList>());
+
+	TArray<UInterchangeBaseNode*> SelectedNodes;
+	const auto NumSelectedItems = GetSelectedItems(SelectedNodes);
+
+	// We always create a section here, even if there is no parent so that clients can still extend the menu
+	MenuBuilder.BeginSection("FbxSceneTreeViewContextMenuImportSection");
+	{
+		const FSlateIcon PlusIcon(FEditorStyle::GetStyleSetName(), "Icons.Plus");
+		MenuBuilder.AddMenuEntry(LOCTEXT("CheckForImport", "Add Selection To Import"), FText(), PlusIcon, FUIAction(FExecuteAction::CreateSP(this, &SInterchangeGraphInspectorTreeView::AddSelectionToImport)));
+		const FSlateIcon MinusIcon(FEditorStyle::GetStyleSetName(), "Icons.Minus");
+		MenuBuilder.AddMenuEntry(LOCTEXT("UncheckForImport", "Remove Selection From Import"), FText(), MinusIcon, FUIAction(FExecuteAction::CreateSP(this, &SInterchangeGraphInspectorTreeView::RemoveSelectionFromImport)));
+	}
+	MenuBuilder.EndSection();
+
+	return MenuBuilder.MakeWidget();
+}
+
+void SInterchangeGraphInspectorTreeView::AddSelectionToImport()
+{
+	SetSelectionImportState(true);
+}
+
+void SInterchangeGraphInspectorTreeView::RemoveSelectionFromImport()
+{
+	SetSelectionImportState(false);
+}
+
+void SInterchangeGraphInspectorTreeView::SetSelectionImportState(bool MarkForImport)
+{
+	TArray<UInterchangeBaseNode*> SelectedNodes;
+	GetSelectedItems(SelectedNodes);
+	for (UInterchangeBaseNode* Item : SelectedNodes)
+	{
+		Item->SetEnabled(MarkForImport);
+	}
+}
+
+void SInterchangeGraphInspectorTreeView::OnTreeViewSelectionChanged(UInterchangeBaseNode* Item, ESelectInfo::Type SelectionType)
+{
+	if (SelectionMode.Get() == ESelectionMode::None)
+	{
+		return;
+	}
+
+	if (OnSelectionChangedDelegate.IsBound())
+	{
+		OnSelectionChangedDelegate.ExecuteIfBound(Item, SelectionType);
+	}
+}
+
+
+/************************************************************************/
+/* SInterchangeGraphInspectorWindow Implementation                      */
+/************************************************************************/
+
+SInterchangeGraphInspectorWindow::SInterchangeGraphInspectorWindow()
+{
+	InterchangeBaseNodeContainer = nullptr;
+	GraphInspectorTreeview = nullptr;
+	GraphInspectorDetailsView = nullptr;
+	OwnerWindow = nullptr;
+}
+
+SInterchangeGraphInspectorWindow::~SInterchangeGraphInspectorWindow()
+{
+	InterchangeBaseNodeContainer = nullptr;
+	GraphInspectorTreeview = nullptr;
+	GraphInspectorDetailsView = nullptr;
+	OwnerWindow = nullptr;
+}
+
+TSharedRef<SBox> SInterchangeGraphInspectorWindow::SpawnGraphInspector()
+{
+	//Create the treeview
+	GraphInspectorTreeview = SNew(SInterchangeGraphInspectorTreeView)
+		.InterchangeBaseNodeContainer(InterchangeBaseNodeContainer)
+		.OnSelectionChangedDelegate(this, &SInterchangeGraphInspectorWindow::OnSelectionChanged);
+
+	TSharedPtr<SBox> InspectorBox;
+	TSharedRef<SBox> GraphInspectorPanelBox = SNew(SBox)
+	[
+		SNew(SSplitter)
+		.Orientation(Orient_Horizontal)
+		+ SSplitter::Slot()
+		.Value(0.4f)
+		[
+			SNew(SVerticalBox)
+			+ SVerticalBox::Slot()
+			.HAlign(HAlign_Left)
+			.AutoHeight()
+			[
+				SNew(SUniformGridPanel)
+				.SlotPadding(2)
+				+ SUniformGridPanel::Slot(0, 0)
+				[
+					SNew(SHorizontalBox)
+					+ SHorizontalBox::Slot()
+					.AutoWidth()
+					[
+						SNew(SCheckBox)
+						.HAlign(HAlign_Center)
+						.OnCheckStateChanged(GraphInspectorTreeview.Get(), &SInterchangeGraphInspectorTreeView::OnToggleSelectAll)
+					]
+					+ SHorizontalBox::Slot()
+					.FillWidth(1.0f)
+					.Padding(0.0f, 3.0f, 6.0f, 3.0f)
+					.VAlign(VAlign_Center)
+					[
+						SNew(STextBlock)
+						.Text(LOCTEXT("GraphInspectorWindow_Scene_All", "All"))
+					]
+				]
+				+ SUniformGridPanel::Slot(1, 0)
+				[
+					SNew(SButton)
+					.HAlign(HAlign_Center)
+					.Text(LOCTEXT("FbxOptionWindow_Scene_ExpandAll", "Expand All"))
+					.OnClicked(GraphInspectorTreeview.Get(), &SInterchangeGraphInspectorTreeView::OnExpandAll)
+				]
+				+ SUniformGridPanel::Slot(2, 0)
+				[
+					SNew(SButton)
+					.HAlign(HAlign_Center)
+					.Text(LOCTEXT("FbxOptionWindow_Scene_CollapseAll", "Collapse All"))
+					.OnClicked(GraphInspectorTreeview.Get(), &SInterchangeGraphInspectorTreeView::OnCollapseAll)
+				]
+			]
+			+ SVerticalBox::Slot()
+			.FillHeight(1.0f)
+			[
+				SNew(SBox)
+				[
+					GraphInspectorTreeview.ToSharedRef()
+				]
+			]
+		]
+		+ SSplitter::Slot()
+		.Value(0.6f)
+		[
+			SAssignNew(InspectorBox, SBox)
+		]
+	];
+
+	FPropertyEditorModule& PropertyEditorModule = FModuleManager::GetModuleChecked<FPropertyEditorModule>("PropertyEditor");
+	FDetailsViewArgs DetailsViewArgs;
+	DetailsViewArgs.bAllowSearch = false;
+	DetailsViewArgs.NameAreaSettings = FDetailsViewArgs::HideNameArea;
+	GraphInspectorDetailsView = PropertyEditorModule.CreateDetailView(DetailsViewArgs);
+	InspectorBox->SetContent(GraphInspectorDetailsView->AsShared());
+	GraphInspectorDetailsView->SetObject(nullptr);
+	return GraphInspectorPanelBox;
+}
+
+
+void SInterchangeGraphInspectorWindow::Construct(const FArguments& InArgs)
+{
+	InterchangeBaseNodeContainer = InArgs._InterchangeBaseNodeContainer;
+	OwnerWindow = InArgs._OwnerWindow;
+
+	check(InterchangeBaseNodeContainer != nullptr);
+	check(OwnerWindow.IsValid());
+
+	this->ChildSlot
+	[
+		SNew(SBorder)
+		.Padding(FMargin(10.0f, 3.0f))
+		.BorderImage(FEditorStyle::GetBrush("ToolPanel.DarkGroupBorder"))
+		[
+			SNew(SVerticalBox)
+			+ SVerticalBox::Slot()
+			.FillHeight(1.0f)
+			.Padding(2)
+			[
+				SpawnGraphInspector()
+			]
+			+ SVerticalBox::Slot()
+			.AutoHeight()
+			.HAlign(HAlign_Right)
+			.Padding(2)
+			[
+				SNew(SUniformGridPanel)
+				.SlotPadding(2)
+				+ SUniformGridPanel::Slot(0, 0)
+				[
+					IDocumentation::Get()->CreateAnchor(FString("Engine/Content/Interchange/GraphInspector"))
+				]
+				+ SUniformGridPanel::Slot(1, 0)
+				[
+					SNew(SButton)
+					.HAlign(HAlign_Center)
+					.Text(LOCTEXT("InspectorGraphWindow_Close", "Ok"))
+					.OnClicked(this, &SInterchangeGraphInspectorWindow::OnCloseDialog)
+				]
+			]
+		]
+	];
+}
+
+void SInterchangeGraphInspectorWindow::OnSelectionChanged(UInterchangeBaseNode* Item, ESelectInfo::Type SelectionType)
+{
+	//Change the object point by the InspectorBox
+	GraphInspectorDetailsView->SetObject(Item);
+}
+
+void SInterchangeGraphInspectorWindow::CloseGraphInspector()
+{
+	InterchangeBaseNodeContainer = nullptr;
+	GraphInspectorTreeview = nullptr;
+	GraphInspectorDetailsView = nullptr;
+
+	if (TSharedPtr<SWindow> OwnerWindowPin = OwnerWindow.Pin())
+	{
+		OwnerWindowPin->RequestDestroyWindow();
+	}
+	OwnerWindow = nullptr;
+}
+
+#undef LOCTEXT_NAMESPACE
diff --git a/Engine/Plugins/Editor/InterchangeEditor/Source/Pipelines/Private/SInterchangeGraphInspectorWindow.h b/Engine/Plugins/Editor/InterchangeEditor/Source/Pipelines/Private/SInterchangeGraphInspectorWindow.h
new file mode 100644
--- /dev/null
+++ b/Engine/Plugins/Editor/InterchangeEditor/Source/Pipelines/Private/SInterchangeGraphInspectorWindow.h
@@ -0,0 +1,111 @@
+// Copyright Epic Games, Inc. All Rights Reserved.
+
+#pragma once
+
+#include "CoreMinimal.h"
+#include "InputCoreTypes.h"
+#include "Input/Reply.h"
+#include "Nodes/InterchangeBaseNode.h"
+#include "Nodes/InterchangeBaseNodeContainer.h"
+#include "Styling/SlateColor.h"
+#include "Widgets/DeclarativeSyntaxSupport.h"
+#include "Widgets/SCompoundWidget.h"
+#include "Widgets/Views/STreeView.h"
+
+class UInterchangeBaseNodeContainer;
+
+DECLARE_DELEGATE_TwoParams(FOnGraphInspectorSelectionChanged, UInterchangeBaseNode*, ESelectInfo::Type)
+
+class SInterchangeGraphInspectorTreeView : public STreeView<UInterchangeBaseNode*>
+{
+public:
+	~SInterchangeGraphInspectorTreeView();
+
+	SLATE_BEGIN_ARGS(SInterchangeGraphInspectorTreeView)
+		: _InterchangeBaseNodeContainer(nullptr)
+		, _OnSelectionChangedDelegate()
+		{}
+		SLATE_ARGUMENT(UInterchangeBaseNodeContainer*, InterchangeBaseNodeContainer)
+		SLATE_EVENT(FOnGraphInspectorSelectionChanged, OnSelectionChangedDelegate)
+	SLATE_END_ARGS()
+
+		/** Construct this widget */
+	void Construct(const FArguments& InArgs);
+	TSharedRef< ITableRow > OnGenerateRowGraphInspectorTreeView(UInterchangeBaseNode* Item, const TSharedRef< STableViewBase >& OwnerTable);
+	void OnGetChildrenGraphInspectorTreeView(UInterchangeBaseNode* InParent, TArray< UInterchangeBaseNode* >& OutChildren);
+
+	void OnToggleSelectAll(ECheckBoxState CheckType);
+	FReply OnExpandAll();
+	FReply OnCollapseAll();
+protected:
+	UInterchangeBaseNodeContainer* InterchangeBaseNodeContainer = nullptr;
+	/** Delegate to invoke when selection changes. */
+	FOnGraphInspectorSelectionChanged OnSelectionChangedDelegate;
+
+	/** the elements we show in the tree view */
+	TArray<UInterchangeBaseNode*> RootNodeArray;
+
+	/** Open a context menu for the current selection */
+	TSharedPtr<SWidget> OnOpenContextMenu();
+	void RecursiveSetImport(UInterchangeBaseNode* Node, bool bState);
+	void RecursiveSetExpand(UInterchangeBaseNode* Node, bool ExpandState);
+	void AddSelectionToImport();
+	void RemoveSelectionFromImport();
+	void SetSelectionImportState(bool MarkForImport);
+	void OnTreeViewSelectionChanged(UInterchangeBaseNode* Item, ESelectInfo::Type SelectionType);
+};
+
+class SInterchangeGraphInspectorWindow : public SCompoundWidget
+{
+public:
+	SLATE_BEGIN_ARGS(SInterchangeGraphInspectorWindow)
+		: _InterchangeBaseNodeContainer(nullptr)
+		, _OwnerWindow()
+		{}
+
+		SLATE_ARGUMENT( UInterchangeBaseNodeContainer*, InterchangeBaseNodeContainer)
+		SLATE_ARGUMENT(TWeakPtr<SWindow>, OwnerWindow)
+	SLATE_END_ARGS()
+
+public:
+
+	SInterchangeGraphInspectorWindow();
+	~SInterchangeGraphInspectorWindow();
+
+	void Construct(const FArguments& InArgs);
+	virtual bool SupportsKeyboardFocus() const override { return true; }
+
+	void CloseGraphInspector();
+
+	FReply OnCloseDialog()
+	{
+		CloseGraphInspector();
+		return FReply::Handled();
+	}
+
+	virtual FReply OnKeyDown( const FGeometry& MyGeometry, const FKeyEvent& InKeyEvent ) override
+	{
+		if( InKeyEvent.GetKey() == EKeys::Escape )
+		{
+			return OnCloseDialog();
+		}
+
+		return FReply::Unhandled();
+	}
+
+private:
+	TWeakPtr< SWindow > OwnerWindow;
+
+	TSharedRef<SBox> SpawnGraphInspector();
+	void OnSelectionChanged(UInterchangeBaseNode* Item, ESelectInfo::Type SelectionType);
+
+	//SInterchangeGraphInspectorWindow Arguments
+	
+	UInterchangeBaseNodeContainer* InterchangeBaseNodeContainer;
+	
+	
+	//Graph Inspector UI elements
+
+	TSharedPtr<SInterchangeGraphInspectorTreeView> GraphInspectorTreeview;
+	TSharedPtr<IDetailsView> GraphInspectorDetailsView;
+};
diff --git a/Engine/Plugins/Editor/InterchangeEditor/Source/Pipelines/Public/InterchangeEditorPipelinesModule.h b/Engine/Plugins/Editor/InterchangeEditor/Source/Pipelines/Public/InterchangeEditorPipelinesModule.h
new file mode 100644
--- /dev/null
+++ b/Engine/Plugins/Editor/InterchangeEditor/Source/Pipelines/Public/InterchangeEditorPipelinesModule.h
@@ -0,0 +1,41 @@
+// Copyright Epic Games, Inc. All Rights Reserved.
+
+#pragma once
+
+#include "CoreMinimal.h"
+#include "Modules/ModuleInterface.h"
+#include "Modules/ModuleManager.h"
+
+
+/**
+ * The public interface to this module
+ */
+class IInterchangeEditorPipelinesModule : public IModuleInterface
+{
+
+public:
+
+	/**
+	 * Singleton-like access to this module's interface.  This is just for convenience!
+	 * Beware of calling this during the shutdown phase, though.  Your module might have been unloaded already.
+	 *
+	 * @return Returns singleton instance, loading the module on demand if needed
+	 */
+	static inline IInterchangeEditorPipelinesModule& Get()
+	{
+		return FModuleManager::LoadModuleChecked< IInterchangeEditorPipelinesModule >( "InterchangeEditorPipelines" );
+	}
+
+	/**
+	 * Checks to see if this module is loaded and ready.  It is only valid to call Get() if IsAvailable() returns true.
+	 *
+	 * @return True if the module is loaded and ready to use
+	 */
+	static inline bool IsAvailable()
+	{
+		return FModuleManager::Get().IsModuleLoaded( "InterchangeEditorPipelines" );
+	}
+
+	TArray<FName> ClassesToUnregisterOnShutdown;
+};
+
diff --git a/Engine/Plugins/Editor/InterchangeEditor/Source/Pipelines/Public/InterchangeGraphInspectorPipeline.h b/Engine/Plugins/Editor/InterchangeEditor/Source/Pipelines/Public/InterchangeGraphInspectorPipeline.h
new file mode 100644
--- /dev/null
+++ b/Engine/Plugins/Editor/InterchangeEditor/Source/Pipelines/Public/InterchangeGraphInspectorPipeline.h
@@ -0,0 +1,53 @@
+// Copyright Epic Games, Inc. All Rights Reserved.
+
+#pragma once
+
+#include "CoreMinimal.h"
+#include "InterchangePipelineBase.h"
+#include "InterchangeSourceData.h"
+#include "Nodes/InterchangeBaseNodeContainer.h"
+#include "UObject/Object.h"
+#include "UObject/ObjectMacros.h"
+
+#include "InterchangeGraphInspectorPipeline.generated.h"
+
+/**
+ * This pipeline is the generic pipeline option for all meshes type and should be call before specialized Mesh pipeline (like generic static mesh or skeletal mesh pipelines)
+ * All shared import options between mesh type should be added here.
+ *
+ * UPROPERTY possible meta values:
+ * @meta ImportOnly - Boolean, the property is use only when we import not when we re-import. Cannot be mix with ReimportOnly!
+ * @meta ReimportOnly - Boolean, the property is use only when we re-import not when we import. Cannot be mix with ImportOnly!
+ * @meta MeshType - String, the property is for static or skeletal or both (static | skeletal) mesh type. If not specified it will apply to all mesh type.
+ */
+UCLASS(BlueprintType)
+class INTERCHANGEEDITORPIPELINES_API UInterchangeGraphInspectorPipeline : public UInterchangePipelineBase
+{
+	GENERATED_BODY()
+
+public:
+	//////////////////////////////////////////////////////////////////////////
+	// BEGIN Pre import pipeline properties
+
+
+	// END Pre import pipeline properties
+	//////////////////////////////////////////////////////////////////////////
+
+protected:
+
+	virtual bool ExecutePreImportPipeline(UInterchangeBaseNodeContainer* BaseNodeContainer, const TArray<UInterchangeSourceData*>& SourceDatas) override;
+
+	//virtual bool ExecutePostImportPipeline(const UInterchangeBaseNodeContainer* BaseNodeContainer, const FName& NodeKey, UObject* CreatedAsset) override;
+
+	virtual bool CanExecuteOnAnyThread(EInterchangePipelineTask PipelineTask) override
+	{
+		//If a blueprint or python derived from this class, it will be execute on the game thread since we cannot currently execute script outside of the game thread, even if this function return true.
+		return false;
+	}
+
+	//virtual bool ExecuteExportPipeline(UInterchangeBaseNodeContainer* BaseNodeContainer) override;
+private:
+
+};
+
+
diff --git a/Engine/Plugins/Runtime/Interchange/Source/Nodes/Public/InterchangeMaterialNode.h b/Engine/Plugins/Runtime/Interchange/Source/Nodes/Public/InterchangeMaterialNode.h
--- a/Engine/Plugins/Runtime/Interchange/Source/Nodes/Public/InterchangeMaterialNode.h
+++ b/Engine/Plugins/Runtime/Interchange/Source/Nodes/Public/InterchangeMaterialNode.h
@@ -108,6 +108,16 @@ public:
 		return TypeName;
 	}
 
+	virtual FString GetKeyDisplayName(const UE::Interchange::FAttributeKey& NodeAttributeKey) const override
+	{
+		FString KeyDisplayName = NodeAttributeKey.ToString();
+		if (NodeAttributeKey == UE::Interchange::FMaterialNodeStaticData::PayloadSourceFileKey())
+		{
+			return KeyDisplayName = TEXT("Payload Source Key");
+		}
+		return Super::GetKeyDisplayName(NodeAttributeKey);
+	}
+
 	/** Get the class this node want to create */
 	UFUNCTION(BlueprintCallable, Category = "Interchange | Node | Material")
 	virtual class UClass* GetAssetClass() const override
@@ -384,7 +394,7 @@ private:
 	//This member is serialize manually in the serialize override
 	TMap<EInterchangeMaterialNodeParameterName, FParameterData> ParameterDatas;
 
-	const UE::Interchange::FAttributeKey ClassNameAttributeKey = UE::Interchange::FAttributeKey(TEXT("__ClassTypeAttribute__"));
+	const UE::Interchange::FAttributeKey ClassNameAttributeKey = UE::Interchange::FBaseNodeStaticData::ClassTypeAttributeKey();
 
 	//Material Attribute
 	const UE::Interchange::FAttributeKey Macro_CustomMaterialUsageKey = UE::Interchange::FAttributeKey(TEXT("MaterialUsage"));
diff --git a/Engine/Plugins/Runtime/Interchange/Source/Nodes/Public/InterchangeSkeletalMeshNode.h b/Engine/Plugins/Runtime/Interchange/Source/Nodes/Public/InterchangeSkeletalMeshNode.h
--- a/Engine/Plugins/Runtime/Interchange/Source/Nodes/Public/InterchangeSkeletalMeshNode.h
+++ b/Engine/Plugins/Runtime/Interchange/Source/Nodes/Public/InterchangeSkeletalMeshNode.h
@@ -83,6 +83,28 @@ public:
 		return TypeName;
 	}
 
+	virtual FString GetKeyDisplayName(const UE::Interchange::FAttributeKey& NodeAttributeKey) const override
+	{
+		FString KeyDisplayName = NodeAttributeKey.Key;
+		if (NodeAttributeKey.Key.Equals(UE::Interchange::FSkeletalMeshNodeStaticData::GetLodDependenciesBaseKey()))
+		{
+			KeyDisplayName = TEXT("LOD Dependencies Count");
+			return KeyDisplayName;
+		}
+		else if (NodeAttributeKey.Key.StartsWith(UE::Interchange::FSkeletalMeshNodeStaticData::GetLodDependenciesBaseKey()))
+		{
+			KeyDisplayName = TEXT("LOD Dependencies Index ");
+			const FString IndexKey = UE::Interchange::FNameAttributeArrayHelper::IndexKey();
+			int32 IndexPosition = NodeAttributeKey.Key.Find(IndexKey) + IndexKey.Len();
+			if (IndexPosition < NodeAttributeKey.Key.Len())
+			{
+				KeyDisplayName += NodeAttributeKey.Key.RightChop(IndexPosition);
+			}
+			return KeyDisplayName;
+		}
+		return Super::GetKeyDisplayName(NodeAttributeKey);
+	}
+
 	/** Get the class this node want to create */
 	UFUNCTION(BlueprintCallable, Category = "Interchange | Node | SkeletalMesh")
 	virtual class UClass* GetAssetClass() const override
@@ -160,7 +182,7 @@ private:
 #endif
 	}
 
-	const UE::Interchange::FAttributeKey ClassNameAttributeKey = UE::Interchange::FAttributeKey(TEXT("__ClassTypeAttribute__"));
+	const UE::Interchange::FAttributeKey ClassNameAttributeKey = UE::Interchange::FBaseNodeStaticData::ClassTypeAttributeKey();
 
 	UE::Interchange::FNameAttributeArrayHelper LodDependencies;
 protected:
diff --git a/Engine/Plugins/Runtime/Interchange/Source/Nodes/Public/InterchangeSkeletonNode.h b/Engine/Plugins/Runtime/Interchange/Source/Nodes/Public/InterchangeSkeletonNode.h
--- a/Engine/Plugins/Runtime/Interchange/Source/Nodes/Public/InterchangeSkeletonNode.h
+++ b/Engine/Plugins/Runtime/Interchange/Source/Nodes/Public/InterchangeSkeletonNode.h
@@ -128,7 +128,7 @@ private:
 #endif
 	}
 
-	const UE::Interchange::FAttributeKey ClassNameAttributeKey = UE::Interchange::FAttributeKey(TEXT("__ClassTypeAttribute__"));
+	const UE::Interchange::FAttributeKey ClassNameAttributeKey = UE::Interchange::FBaseNodeStaticData::ClassTypeAttributeKey();
 
 	//Skeleton
 	const UE::Interchange::FAttributeKey Macro_CustomRootJointIDKey = UE::Interchange::FAttributeKey(TEXT("RootJointID"));
diff --git a/Engine/Plugins/Runtime/Interchange/Source/Nodes/Public/InterchangeTextureNode.h b/Engine/Plugins/Runtime/Interchange/Source/Nodes/Public/InterchangeTextureNode.h
--- a/Engine/Plugins/Runtime/Interchange/Source/Nodes/Public/InterchangeTextureNode.h
+++ b/Engine/Plugins/Runtime/Interchange/Source/Nodes/Public/InterchangeTextureNode.h
@@ -26,6 +26,12 @@ namespace UE
 				static FAttributeKey AttributeKey(TEXT("__PayloadSourceFile__"));
 				return AttributeKey;
 			}
+
+			static const FString& GetBaseSourceBlocksKey()
+			{
+				static FString StringKey(TEXT("SourceBlocks"));
+				return StringKey;
+			}
 		};
 
 	}//ns Interchange
@@ -48,7 +54,7 @@ public:
 	virtual void PostInitProperties()
 	{
 		Super::PostInitProperties();
-		SourceBlocks.Initialize(Attributes.ToSharedRef(), TEXT("SourceBlocks"));
+		SourceBlocks.Initialize(Attributes.ToSharedRef(), UE::Interchange::FTextureNodeStaticData::GetBaseSourceBlocksKey());
 	}
 
 	/**
@@ -89,6 +95,16 @@ public:
 		return TypeName;
 	}
 
+	virtual FString GetKeyDisplayName(const UE::Interchange::FAttributeKey& NodeAttributeKey) const override
+	{
+		FString KeyDisplayName = NodeAttributeKey.ToString();
+		if (NodeAttributeKey == UE::Interchange::FTextureNodeStaticData::PayloadSourceFileKey())
+		{
+			return KeyDisplayName = TEXT("Payload Source Key");
+		}
+		return Super::GetKeyDisplayName(NodeAttributeKey);
+	}
+
 	/** Get the class this node want to create */
 	UFUNCTION(BlueprintCallable, Category = "Interchange | Node | Texture")
 	virtual class UClass* GetAssetClass() const override
@@ -817,7 +833,7 @@ private:
 #endif
 	}
 
-	const UE::Interchange::FAttributeKey ClassNameAttributeKey = UE::Interchange::FAttributeKey(TEXT("__ClassTypeAttribute__"));
+	const UE::Interchange::FAttributeKey ClassNameAttributeKey = UE::Interchange::FBaseNodeStaticData::ClassTypeAttributeKey();
 
 	//Texture Adjustments
 	const UE::Interchange::FAttributeKey Macro_CustomAdjustBrightnessKey = UE::Interchange::FAttributeKey(TEXT("AdjustBrightness"));
diff --git a/Engine/Source/Runtime/Interchange/Core/Private/Nodes/InterchangeBaseNode.cpp b/Engine/Source/Runtime/Interchange/Core/Private/Nodes/InterchangeBaseNode.cpp
--- a/Engine/Source/Runtime/Interchange/Core/Private/Nodes/InterchangeBaseNode.cpp
+++ b/Engine/Source/Runtime/Interchange/Core/Private/Nodes/InterchangeBaseNode.cpp
@@ -119,6 +119,11 @@ bool UInterchangeBaseNode::RemoveDependencyUID(const FString& DependencyUID)
 
 bool UInterchangeBaseNode::IsEnabled() const
 {
+	if (!Attributes->ContainAttribute(UE::Interchange::FBaseNodeStaticData::IsEnabledKey()))
+	{
+		return false;
+	}
+
 	UE::Interchange::FAttributeStorage::TAttributeHandle<bool> Handle = Attributes->GetAttributeHandle<bool>(UE::Interchange::FBaseNodeStaticData::IsEnabledKey());
 	if (Handle.IsValid())
 	{
diff --git a/Engine/Source/Runtime/Interchange/Core/Private/Types/AttributeStorage.cpp b/Engine/Source/Runtime/Interchange/Core/Private/Types/AttributeStorage.cpp
--- a/Engine/Source/Runtime/Interchange/Core/Private/Types/AttributeStorage.cpp
+++ b/Engine/Source/Runtime/Interchange/Core/Private/Types/AttributeStorage.cpp
@@ -11,6 +11,70 @@ namespace UE
 {
 	namespace Interchange
 	{
+		FString AttributeTypeToString(EAttributeTypes AttributeType)
+		{
+			FString AttributeTypeString;
+			switch (AttributeType)
+			{
+				case EAttributeTypes::None: { AttributeTypeString = TEXT("None"); } break;
+				case EAttributeTypes::Bool: { AttributeTypeString = TEXT("Bool"); } break;
+				case EAttributeTypes::Box: { AttributeTypeString = TEXT("Box"); } break;
+				case EAttributeTypes::BoxSphereBounds: { AttributeTypeString = TEXT("BoxSphereBounds"); } break;
+				case EAttributeTypes::ByteArray: { AttributeTypeString = TEXT("ByteArray"); } break;
+				case EAttributeTypes::ByteArray64: { AttributeTypeString = TEXT("ByteArray64"); } break;
+				case EAttributeTypes::Color: { AttributeTypeString = TEXT("Color"); } break;
+				case EAttributeTypes::DateTime: { AttributeTypeString = TEXT("DateTime"); } break;
+				case EAttributeTypes::Double: { AttributeTypeString = TEXT("Double"); } break;
+				case EAttributeTypes::Enum: { AttributeTypeString = TEXT("Enum"); } break;
+				case EAttributeTypes::Float: { AttributeTypeString = TEXT("Float"); } break;
+				case EAttributeTypes::Guid: { AttributeTypeString = TEXT("Guid"); } break;
+				case EAttributeTypes::Int8: { AttributeTypeString = TEXT("Int8"); } break;
+				case EAttributeTypes::Int16: { AttributeTypeString = TEXT("Int16"); } break;
+				case EAttributeTypes::Int32: { AttributeTypeString = TEXT("Int32"); } break;
+				case EAttributeTypes::Int64: { AttributeTypeString = TEXT("Int64"); } break;
+				case EAttributeTypes::IntRect: { AttributeTypeString = TEXT("IntRect"); } break;
+				case EAttributeTypes::LinearColor: { AttributeTypeString = TEXT("LinearColor"); } break;
+				case EAttributeTypes::Matrix: { AttributeTypeString = TEXT("Matrix"); } break;
+				case EAttributeTypes::Name: { AttributeTypeString = TEXT("Name"); } break;
+				case EAttributeTypes::Plane: { AttributeTypeString = TEXT("Plane"); } break;
+				case EAttributeTypes::Quat: { AttributeTypeString = TEXT("Quat"); } break;
+				case EAttributeTypes::RandomStream: { AttributeTypeString = TEXT("RandomStream"); } break;
+				case EAttributeTypes::Rotator: { AttributeTypeString = TEXT("Rotator"); } break;
+				case EAttributeTypes::String: { AttributeTypeString = TEXT("String"); } break;
+				case EAttributeTypes::Timespan: { AttributeTypeString = TEXT("Timespan"); } break;
+				case EAttributeTypes::Transform: { AttributeTypeString = TEXT("Transform"); } break;
+				case EAttributeTypes::TwoVectors: { AttributeTypeString = TEXT("TwoVectors"); } break;
+				case EAttributeTypes::UInt8: { AttributeTypeString = TEXT("UInt8"); } break;
+				case EAttributeTypes::UInt16: { AttributeTypeString = TEXT("UInt16"); } break;
+				case EAttributeTypes::UInt32: { AttributeTypeString = TEXT("UInt32"); } break;
+				case EAttributeTypes::UInt64: { AttributeTypeString = TEXT("UInt64"); } break;
+				case EAttributeTypes::Vector: { AttributeTypeString = TEXT("Vector"); } break;
+				case EAttributeTypes::Vector2d: { AttributeTypeString = TEXT("Vector2d"); } break;
+				case EAttributeTypes::Vector4: { AttributeTypeString = TEXT("Vector4"); } break;
+				case EAttributeTypes::IntPoint: { AttributeTypeString = TEXT("IntPoint"); } break;
+				case EAttributeTypes::IntVector: { AttributeTypeString = TEXT("IntVector"); } break;
+				case EAttributeTypes::NetworkGUID: { AttributeTypeString = TEXT("NetworkGUID"); } break;
+				case EAttributeTypes::Vector2DHalf: { AttributeTypeString = TEXT("Vector2DHalf"); } break;
+				case EAttributeTypes::Float16: { AttributeTypeString = TEXT("Float16"); } break;
+				case EAttributeTypes::CapsuleShape: { AttributeTypeString = TEXT("CapsuleShape"); } break;
+				case EAttributeTypes::DualQuat: { AttributeTypeString = TEXT("DualQuat"); } break;
+				case EAttributeTypes::OrientedBox: { AttributeTypeString = TEXT("OrientedBox"); } break;
+				case EAttributeTypes::Ray: { AttributeTypeString = TEXT("Ray"); } break;
+				case EAttributeTypes::Sphere: { AttributeTypeString = TEXT("Sphere"); } break;
+				case EAttributeTypes::FrameNumber: { AttributeTypeString = TEXT("FrameNumber"); } break;
+				case EAttributeTypes::FrameRate: { AttributeTypeString = TEXT("FrameRate"); } break;
+				case EAttributeTypes::FrameTime: { AttributeTypeString = TEXT("FrameTime"); } break;
+				default:
+				{
+					//Ensure if we ask an unknown type
+					const bool bUnknownType = true;
+					ensure(!bUnknownType);
+				}
+				break;
+			}
+			return AttributeTypeString;
+		}
+
 		void LogAttributeStorageErrors(const EAttributeStorageResult Result, const FString OperationName, const FAttributeKey AttributeKey)
 		{
 			//////////////////////////////////////////////////////////////////////////
diff --git a/Engine/Source/Runtime/Interchange/Core/Public/Nodes/InterchangeBaseNode.h b/Engine/Source/Runtime/Interchange/Core/Public/Nodes/InterchangeBaseNode.h
--- a/Engine/Source/Runtime/Interchange/Core/Public/Nodes/InterchangeBaseNode.h
+++ b/Engine/Source/Runtime/Interchange/Core/Public/Nodes/InterchangeBaseNode.h
@@ -187,10 +187,16 @@ namespace UE
 			{
 				Attributes = InAttributes;
 				check(Attributes.IsValid());
-				FString BaseTryName = TEXT("__") + BaseKeyName;
+				FString BaseTryName = BaseKeyName;
 				KeyCount = BaseTryName;
 			}
 
+			static const FString& IndexKey()
+			{
+				static FString IndexKeyString = TEXT("_NameIndex_");
+				return IndexKeyString;
+			}
+
 			int32 GetCount() const
 			{
 				//The Class must be initialise properly before we can use it
@@ -392,7 +398,7 @@ namespace UE
 
 			FAttributeKey GetIndexKey(int32 Index) const
 			{
-				FString DepIndexKeyString = GetKeyCount().ToString() + TEXT("_NameIndex_") + FString::FromInt(Index);
+				FString DepIndexKeyString = GetKeyCount().ToString() + IndexKey() + FString::FromInt(Index);
 				return FAttributeKey(*DepIndexKeyString);
 			}
 		};
@@ -410,7 +416,7 @@ namespace UE
 			void Initialize(const TSharedRef<FAttributeStorage, ESPMode::ThreadSafe>& InAttributes, const FString& BaseKeyName)
 			{
 				Attributes = InAttributes;
-				FString BaseTryName = TEXT("__") + BaseKeyName;
+				FString BaseTryName = BaseKeyName;
 				FAttributeKey KeyCountKey(MoveTemp(BaseTryName));
 				if (!InAttributes->ContainAttribute(KeyCountKey))
 				{ 
@@ -713,9 +719,15 @@ namespace UE
 
 			static const FString& GetDependenciesBaseKey()
 			{
-				static FString BaseNodeDependencies_BaseKey = TEXT("BaseNodeDependencies__");
+				static FString BaseNodeDependencies_BaseKey = TEXT("__BaseNodeDependencies__");
 				return BaseNodeDependencies_BaseKey;
 			}
+
+			static const FAttributeKey& ClassTypeAttributeKey()
+			{
+				static FAttributeKey AttributeKey(TEXT("__ClassTypeAttribute__"));
+				return AttributeKey;
+			}
 		};
 
 	} //ns Interchange
@@ -736,6 +748,7 @@ public:
 	{
 		Attributes = MakeShared<UE::Interchange::FAttributeStorage, ESPMode::ThreadSafe>();
 		Dependencies.Initialize(Attributes, UE::Interchange::FBaseNodeStaticData::GetDependenciesBaseKey());
+		RegisterAttribute<bool>(UE::Interchange::FBaseNodeStaticData::IsEnabledKey(), true);
 	}
 
 	virtual ~UInterchangeBaseNode() = default;
@@ -754,6 +767,46 @@ public:
 	 */
 	virtual FString GetTypeName() const;
 
+	virtual FString GetKeyDisplayName(const UE::Interchange::FAttributeKey& NodeAttributeKey) const
+	{
+		FString KeyDisplayName = NodeAttributeKey.ToString();
+		if (NodeAttributeKey == UE::Interchange::FBaseNodeStaticData::ParentIDKey())
+		{
+			KeyDisplayName = TEXT("Parent Unique ID");
+		}
+		else if (NodeAttributeKey == UE::Interchange::FBaseNodeStaticData::DisplayLabelKey())
+		{
+			KeyDisplayName = TEXT("Name");
+		}
+		else if (NodeAttributeKey == UE::Interchange::FBaseNodeStaticData::IsEnabledKey())
+		{
+			KeyDisplayName = TEXT("Enabled");
+		}
+		else if (NodeAttributeKey == UE::Interchange::FBaseNodeStaticData::UniqueIDKey())
+		{
+			KeyDisplayName = TEXT("Unique ID");
+		}
+		else if (NodeAttributeKey == UE::Interchange::FBaseNodeStaticData::ClassTypeAttributeKey())
+		{
+			KeyDisplayName = TEXT("Node Class Type");
+		}
+		else if (NodeAttributeKey.Key.Equals(UE::Interchange::FBaseNodeStaticData::GetDependenciesBaseKey()))
+		{
+			KeyDisplayName = TEXT("Dependencies Count");
+		}
+		else if (NodeAttributeKey.Key.StartsWith(UE::Interchange::FBaseNodeStaticData::GetDependenciesBaseKey()))
+		{
+			KeyDisplayName = TEXT("Dependencies Index ");
+			const FString IndexKey = UE::Interchange::FNameAttributeArrayHelper::IndexKey();
+			int32 IndexPosition = NodeAttributeKey.Key.Find(IndexKey) + IndexKey.Len();
+			if (IndexPosition < NodeAttributeKey.Key.Len())
+			{
+				KeyDisplayName += NodeAttributeKey.Key.RightChop(IndexPosition);
+			}
+		}
+		return KeyDisplayName;
+	}
+
 	/**
 	 * Add an attribute to the node
 	 * @param NodeAttributeKey - The key of the attribute
@@ -802,6 +855,11 @@ public:
 		return Attributes->GetAttributeHandle<T>(NodeAttributeKey);
 	}
 
+	void GetAttributeKeys(TArray<UE::Interchange::FAttributeKey>& AttributeKeys)
+	{
+		Attributes->GetAttributeKeys(AttributeKeys);
+	}
+
 	/**
 	 * Return the unique id pass in the constructor.
 	 *
diff --git a/Engine/Source/Runtime/Interchange/Core/Public/Types/AttributeStorage.h b/Engine/Source/Runtime/Interchange/Core/Public/Types/AttributeStorage.h
--- a/Engine/Source/Runtime/Interchange/Core/Public/Types/AttributeStorage.h
+++ b/Engine/Source/Runtime/Interchange/Core/Public/Types/AttributeStorage.h
@@ -194,6 +194,10 @@ namespace UE
 			FrameTime			= 47,
 		};
 
+		/**
+		 * Return the FString for the specified AttributeType
+		 */
+		INTERCHANGECORE_API FString AttributeTypeToString(EAttributeTypes AttributeType);
 
 		/**
 		 * Stub for attribute type traits.
