Fix no unity build broken by the last interchange development

diff --git a/Engine/Plugins/Editor/InterchangeEditor/Source/Pipelines/Private/SInterchangeGraphInspectorWindow.cpp b/Engine/Plugins/Editor/InterchangeEditor/Source/Pipelines/Private/SInterchangeGraphInspectorWindow.cpp
--- a/Engine/Plugins/Editor/InterchangeEditor/Source/Pipelines/Private/SInterchangeGraphInspectorWindow.cpp
+++ b/Engine/Plugins/Editor/InterchangeEditor/Source/Pipelines/Private/SInterchangeGraphInspectorWindow.cpp
@@ -2,8 +2,10 @@
 #include "SInterchangeGraphInspectorWindow.h"
 
 #include "EditorStyleSet.h"
+#include "Framework/MultiBox/MultiBoxBuilder.h"
 #include "IDetailsView.h"
 #include "IDocumentation.h"
+#include "GameFramework/Actor.h"
 #include "Modules/ModuleManager.h"
 #include "Nodes/InterchangeBaseNode.h"
 #include "Nodes/InterchangeBaseNodeContainer.h"
diff --git a/Engine/Plugins/Editor/InterchangeEditor/Source/Pipelines/Private/SInterchangeGraphInspectorWindow.h b/Engine/Plugins/Editor/InterchangeEditor/Source/Pipelines/Private/SInterchangeGraphInspectorWindow.h
--- a/Engine/Plugins/Editor/InterchangeEditor/Source/Pipelines/Private/SInterchangeGraphInspectorWindow.h
+++ b/Engine/Plugins/Editor/InterchangeEditor/Source/Pipelines/Private/SInterchangeGraphInspectorWindow.h
@@ -12,6 +12,7 @@
 #include "Widgets/SCompoundWidget.h"
 #include "Widgets/Views/STreeView.h"
 
+class IDetailsView;
 class UInterchangeBaseNodeContainer;
 
 DECLARE_DELEGATE_TwoParams(FOnGraphInspectorSelectionChanged, UInterchangeBaseNode*, ESelectInfo::Type)
