Pipeline workflows prototype. This is a first implementation of the interchange pipeline selection workflow when users import files with interchange.
- Interchange per project settings allow user to create pipeline stacks and to override the pipeline selector UI
- First draft of the generic pipeline selector UI dialog
- Pipeline settings are saved in the config ini files
- The importAssetData now contains the pipelines settings so we can reuse the same settings when we will do the re import

diff --git a/Engine/Config/BaseEngine.ini b/Engine/Config/BaseEngine.ini
--- a/Engine/Config/BaseEngine.ini
+++ b/Engine/Config/BaseEngine.ini
@@ -2789,7 +2789,10 @@ bDisableActionLimit=false
 ActionLimit=100
 
 [/Script/InterchangeEngine.InterchangeProjectSettings]
-PipelineStack=/Script/InterchangePipelines.InterchangeGenericAssetsPipeline
+PipelineStacks=(("Generic", (Pipelines=("/Script/InterchangePipelines.InterchangeGenericAssetsPipeline"))))
+DefaultPipelineStack=Generic
+PipelineConfigurationDialogClass=/Script/InterchangeEditorPipelines.InterchangePipelineConfigurationGeneric
+bShowPipelineStacksConfigurationDialog=True
 
 [AlternateTextureCompression]
 TextureCompressionFormat="TextureFormatOodle"
diff --git a/Engine/Plugins/Editor/InterchangeEditor/Source/Pipelines/Private/InterchangeEditorPipelineStyle.h b/Engine/Plugins/Editor/InterchangeEditor/Source/Pipelines/Private/InterchangeEditorPipelineStyle.h
--- a/Engine/Plugins/Editor/InterchangeEditor/Source/Pipelines/Private/InterchangeEditorPipelineStyle.h
+++ b/Engine/Plugins/Editor/InterchangeEditor/Source/Pipelines/Private/InterchangeEditorPipelineStyle.h
@@ -24,6 +24,7 @@ public:
 		const FString BaseDir = IPluginManager::Get().FindPlugin("InterchangeEditor")->GetBaseDir();
 		SetContentRoot(BaseDir / TEXT("Content"));
 
+		//GraphInspectorIcon icons
 		FSlateImageBrush* LodBrush16 = new FSlateImageBrush(RootToContentDir(TEXT("Resources/Interchange_Lod_Icon_16"), TEXT(".png")), Icon16x16);
 		Set("SceneGraphIcon.LodGroup", LodBrush16);
 		FSlateImageBrush* JointBrush16 = new FSlateImageBrush(RootToContentDir(TEXT("Resources/Interchange_Joint_Icon_16"), TEXT(".png")), Icon16x16);
@@ -32,6 +33,15 @@ public:
 		Set("MeshIcon.Static", StaticMeshBrush16);
 		FSlateImageBrush* SkeletalMeshBrush16 = new FSlateImageBrush(RootToContentDir(TEXT("Resources/Interchange_SkeletalMesh_Icon_16"), TEXT(".png")), Icon16x16);
 		Set("MeshIcon.Skinned", SkeletalMeshBrush16);
+
+		//PipelineConfigurationIcon icons
+		FSlateImageBrush* PipelineBrush16 = new FSlateImageBrush(RootToContentDir(TEXT("Resources/Interchange_Pipeline_Icon_16"), TEXT(".png")), Icon16x16);
+		Set("PipelineConfigurationIcon.Pipeline", PipelineBrush16);
+		FSlateImageBrush* PipelineStackBrush16 = new FSlateImageBrush(RootToContentDir(TEXT("Resources/Interchange_PipelineStack_Icon_16"), TEXT(".png")), Icon16x16);
+		Set("PipelineConfigurationIcon.PipelineStack", PipelineStackBrush16);
+		FSlateImageBrush* PipelineStackDefaultBrush16 = new FSlateImageBrush(RootToContentDir(TEXT("Resources/Interchange_PipelineStackDefault_Icon_16"), TEXT(".png")), Icon16x16);
+		Set("PipelineConfigurationIcon.PipelineStackDefault", PipelineStackDefaultBrush16);
+
 		FSlateStyleRegistry::RegisterSlateStyle(*this);
 	}
 
diff --git a/Engine/Plugins/Editor/InterchangeEditor/Source/Pipelines/Private/InterchangePipelineConfigurationGeneric.cpp b/Engine/Plugins/Editor/InterchangeEditor/Source/Pipelines/Private/InterchangePipelineConfigurationGeneric.cpp
new file mode 100644
--- /dev/null
+++ b/Engine/Plugins/Editor/InterchangeEditor/Source/Pipelines/Private/InterchangePipelineConfigurationGeneric.cpp
@@ -0,0 +1,45 @@
+// Copyright Epic Games, Inc. All Rights Reserved. 
+#include "InterchangePipelineConfigurationGeneric.h"
+
+#include "CoreMinimal.h"
+#include "Framework/Application/SlateApplication.h"
+#include "Interfaces/IMainFrameModule.h"
+#include "Misc/Paths.h"
+#include "SInterchangePipelineConfigurationDialog.h"
+#include "UObject/Object.h"
+#include "UObject/ObjectMacros.h"
+#include "Widgets/SWindow.h"
+
+EInterchangePipelineConfigurationDialogResult UInterchangePipelineConfigurationGeneric::ShowPipelineConfigurationDialog()
+{
+	//Create and show the graph inspector UI dialog
+	TSharedPtr<SWindow> ParentWindow;
+	if(IMainFrameModule* MainFrame = FModuleManager::LoadModulePtr<IMainFrameModule>("MainFrame"))
+	{
+		ParentWindow = MainFrame->GetParentWindow();
+	}
+	TSharedRef<SWindow> Window = SNew(SWindow)
+		.ClientSize(FVector2D(1000.f, 650.f))
+		.Title(NSLOCTEXT("Interchange", "PipelineConfigurationGenericTitle", "Interchange Pipeline Configuration"));
+	TSharedPtr<SInterchangePipelineConfigurationDialog> InterchangePipelineConfigurationDialog;
+
+	Window->SetContent
+	(
+		SAssignNew(InterchangePipelineConfigurationDialog, SInterchangePipelineConfigurationDialog)
+		.OwnerWindow(Window)
+	);
+
+	FSlateApplication::Get().AddModalWindow(Window, ParentWindow, false);
+	
+	if (InterchangePipelineConfigurationDialog->IsCanceled())
+	{
+		return EInterchangePipelineConfigurationDialogResult::Cancel;
+	}
+	
+	if (InterchangePipelineConfigurationDialog->IsImportAll())
+	{
+		return EInterchangePipelineConfigurationDialogResult::ImportAll;
+	}
+
+	return EInterchangePipelineConfigurationDialogResult::Import;
+}
diff --git a/Engine/Plugins/Editor/InterchangeEditor/Source/Pipelines/Private/SInterchangePipelineConfigurationDialog.cpp b/Engine/Plugins/Editor/InterchangeEditor/Source/Pipelines/Private/SInterchangePipelineConfigurationDialog.cpp
new file mode 100644
--- /dev/null
+++ b/Engine/Plugins/Editor/InterchangeEditor/Source/Pipelines/Private/SInterchangePipelineConfigurationDialog.cpp
@@ -0,0 +1,499 @@
+// Copyright Epic Games, Inc. All Rights Reserved.
+#include "SInterchangePipelineConfigurationDialog.h"
+
+#include "EditorStyleSet.h"
+#include "Framework/MultiBox/MultiBoxBuilder.h"
+#include "GameFramework/Actor.h"
+#include "IDetailsView.h"
+#include "IDocumentation.h"
+#include "InterchangePipelineBase.h"
+#include "InterchangeProjectSettings.h"
+#include "Misc/App.h"
+#include "Modules/ModuleManager.h"
+#include "PropertyEditorModule.h"
+#include "Styling/SlateIconFinder.h"
+#include "Widgets/Images/SImage.h"
+#include "Widgets/Input/SButton.h"
+#include "Widgets/Input/SCheckBox.h"
+#include "Widgets/Input/SEditableText.h"
+#include "Widgets/Input/STextComboBox.h"
+#include "Widgets/Layout/SSplitter.h"
+#include "Widgets/Layout/SUniformGridPanel.h"
+
+#define LOCTEXT_NAMESPACE "InterchangePipelineConfiguration"
+
+/************************************************************************/
+/* FInterchangePipelineStacksTreeNodeItem Implementation                    */
+/************************************************************************/
+
+void FInterchangePipelineStacksTreeNodeItem::AddReferencedObjects(FReferenceCollector& Collector)
+{
+	if (Pipeline)
+	{
+		Collector.AddReferencedObject(Pipeline);
+	}
+}
+
+/************************************************************************/
+/* SInterchangePipelineStacksTreeView Implementation                    */
+/************************************************************************/
+
+SInterchangePipelineStacksTreeView::~SInterchangePipelineStacksTreeView()
+{
+
+}
+
+void SInterchangePipelineStacksTreeView::Construct(const FArguments& InArgs)
+{
+	OnSelectionChangedDelegate = InArgs._OnSelectionChangedDelegate;
+	//Build the FbxNodeInfoPtr tree data
+
+	const FName& DefaultPipelineStackName = GetDefault<UInterchangeProjectSettings>()->DefaultPipelineStack;
+	const TMap<FName, FInterchangePipelineStack>& DefaultPipelineStacks = GetDefault<UInterchangeProjectSettings>()->PipelineStacks;
+
+	 
+	for (const TPair<FName, FInterchangePipelineStack>& NameAndPipelineStack : DefaultPipelineStacks)
+	{
+		TSharedPtr<FInterchangePipelineStacksTreeNodeItem> StackNode = MakeShared<FInterchangePipelineStacksTreeNodeItem>();
+		StackNode->StackName = NameAndPipelineStack.Key;
+		StackNode->Pipeline = nullptr;
+		const FInterchangePipelineStack& PipelineStack = NameAndPipelineStack.Value;
+		for (int32 PipelineIndex = 0; PipelineIndex < PipelineStack.Pipelines.Num(); ++PipelineIndex)
+		{
+			const UClass* PipelineClass = PipelineStack.Pipelines[PipelineIndex].LoadSynchronous();
+			if (PipelineClass)
+			{
+				UInterchangePipelineBase* GeneratedPipeline = NewObject<UInterchangePipelineBase>(GetTransientPackage(), PipelineClass, NAME_None, RF_NoFlags);
+				//Load the settings for this pipeline
+				GeneratedPipeline->LoadSettings(NameAndPipelineStack.Key);
+				TSharedPtr<FInterchangePipelineStacksTreeNodeItem> PipelineNode = MakeShared<FInterchangePipelineStacksTreeNodeItem>();
+				PipelineNode->StackName = NameAndPipelineStack.Key;
+				PipelineNode->Pipeline = GeneratedPipeline;
+				StackNode->Childrens.Add(PipelineNode);
+			}
+		}
+		RootNodeArray.Add(StackNode);
+	}
+
+	STreeView::Construct
+	(
+		STreeView::FArguments()
+		.TreeItemsSource(&RootNodeArray)
+		.SelectionMode(ESelectionMode::Multi)
+		.OnGenerateRow(this, &SInterchangePipelineStacksTreeView::OnGenerateRowPipelineConfigurationTreeView)
+		.OnGetChildren(this, &SInterchangePipelineStacksTreeView::OnGetChildrenPipelineConfigurationTreeView)
+		.OnContextMenuOpening(this, &SInterchangePipelineStacksTreeView::OnOpenContextMenu)
+		.OnSelectionChanged(this, &SInterchangePipelineStacksTreeView::OnTreeViewSelectionChanged)
+	);
+}
+
+/** The item used for visualizing the class in the tree. */
+class SInterchangePipelineStacksTreeViewItem : public STableRow< TSharedPtr<FInterchangePipelineStacksTreeNodeItem> >
+{
+public:
+
+	SLATE_BEGIN_ARGS(SInterchangePipelineStacksTreeViewItem)
+		: _InterchangeNode(nullptr)
+	{}
+
+		/** The item content. */
+		SLATE_ARGUMENT(TSharedPtr<FInterchangePipelineStacksTreeNodeItem>, InterchangeNode)
+	SLATE_END_ARGS()
+
+	/**
+	* Construct the widget
+	*
+	* @param InArgs   A declaration from which to construct the widget
+	*/
+	void Construct(const FArguments& InArgs, const TSharedRef<STableViewBase>& InOwnerTableView)
+	{
+		InterchangeNode = InArgs._InterchangeNode;
+		//This is suppose to always be valid
+		check(InterchangeNode);
+		const bool bIsPipelineStackNode = InterchangeNode->Pipeline == nullptr;
+
+		//Prepare the tooltip
+		FString Tooltip;
+		//FString Tooltip = bIsPipelineStackNode ? InterchangeNode->GetDisplayLabel();
+		FText NodeDisplayLabel = FText::FromName(bIsPipelineStackNode ? InterchangeNode->StackName : InterchangeNode->Pipeline->GetClass()->GetFName());
+		this->ChildSlot
+		[
+			SNew(SHorizontalBox)
+			+ SHorizontalBox::Slot()
+			.AutoWidth()
+			[
+				SNew(SExpanderArrow, SharedThis(this))
+			]
+			+ SHorizontalBox::Slot()
+			.AutoWidth()
+			.Padding(0.0f, 2.0f, 6.0f, 2.0f)
+			[
+				SNew(SImage)
+				.Image(this, &SInterchangePipelineStacksTreeViewItem::GetImageItemIcon)
+			]
+			+ SHorizontalBox::Slot()
+			.FillWidth(1.0f)
+			.Padding(0.0f, 3.0f, 6.0f, 3.0f)
+			.VAlign(VAlign_Center)
+			[
+				SNew(STextBlock)
+				.Text(NodeDisplayLabel)
+				.ToolTipText(FText::FromString(Tooltip))
+			]
+		];
+
+		STableRow< TSharedPtr<FInterchangePipelineStacksTreeNodeItem> >::ConstructInternal(
+			STableRow::FArguments()
+			.ShowSelection(true),
+			InOwnerTableView
+		);
+	}
+
+private:
+	const FSlateBrush* GetImageItemIcon() const
+	{
+		const FName DefaultPipelineStackName = GetDefault<UInterchangeProjectSettings>()->DefaultPipelineStack;
+		//This is suppose to always be valid
+		check(InterchangeNode);
+		const bool bIsPipelineStackNode = InterchangeNode->Pipeline == nullptr;
+		const bool bIsDefaultStackNode = bIsPipelineStackNode && (DefaultPipelineStackName == InterchangeNode->StackName);
+		const FSlateBrush * TypeIcon = nullptr;
+		FName IconName = bIsDefaultStackNode ? "PipelineConfigurationIcon.PipelineStackDefault" : bIsPipelineStackNode ? "PipelineConfigurationIcon.PipelineStack" : "PipelineConfigurationIcon.Pipeline";
+		if (IconName != NAME_None)
+		{
+			const FSlateIcon SlateIcon = FSlateIconFinder::FindIcon(IconName);
+			TypeIcon = SlateIcon.GetOptionalIcon();
+		}
+
+		if (!TypeIcon)
+		{
+			TypeIcon = FSlateIconFinder::FindIconBrushForClass(AActor::StaticClass());
+		}
+		return TypeIcon;
+	}
+
+	/** The node to build the tree view row from. */
+	TSharedPtr<FInterchangePipelineStacksTreeNodeItem> InterchangeNode = nullptr;
+};
+
+TSharedRef< ITableRow > SInterchangePipelineStacksTreeView::OnGenerateRowPipelineConfigurationTreeView(TSharedPtr<FInterchangePipelineStacksTreeNodeItem> Item, const TSharedRef< STableViewBase >& OwnerTable)
+{
+	TSharedRef< SInterchangePipelineStacksTreeViewItem > ReturnRow = SNew(SInterchangePipelineStacksTreeViewItem, OwnerTable)
+		.InterchangeNode(Item);
+	return ReturnRow;
+}
+void SInterchangePipelineStacksTreeView::OnGetChildrenPipelineConfigurationTreeView(TSharedPtr<FInterchangePipelineStacksTreeNodeItem> InParent, TArray< TSharedPtr<FInterchangePipelineStacksTreeNodeItem> >& OutChildren)
+{
+	for (int32 ChildIndex = 0; ChildIndex < InParent->Childrens.Num(); ++ChildIndex)
+	{
+		TSharedPtr<FInterchangePipelineStacksTreeNodeItem> ChildNode = InParent->Childrens[ChildIndex];
+		if (!ChildNode.IsValid())
+			continue;
+		OutChildren.Add(ChildNode);
+	}
+}
+
+void SInterchangePipelineStacksTreeView::RecursiveSetExpand(TSharedPtr<FInterchangePipelineStacksTreeNodeItem> Node, bool ExpandState)
+{
+	SetItemExpansion(Node, ExpandState);
+	for (int32 ChildIndex = 0; ChildIndex < Node->Childrens.Num(); ++ChildIndex)
+	{
+		TSharedPtr<FInterchangePipelineStacksTreeNodeItem> ChildNode = Node->Childrens[ChildIndex];
+		if (!ChildNode.IsValid())
+			continue;
+		RecursiveSetExpand(ChildNode, ExpandState);
+	}
+}
+
+FReply SInterchangePipelineStacksTreeView::OnExpandAll()
+{
+	for (TSharedPtr<FInterchangePipelineStacksTreeNodeItem> Node : RootNodeArray)
+	{
+		if (!ensure(Node))
+		{
+			continue;
+		}
+		RecursiveSetExpand(Node, true);
+	}
+	return FReply::Handled();
+}
+
+FReply SInterchangePipelineStacksTreeView::OnCollapseAll()
+{
+	for (TSharedPtr<FInterchangePipelineStacksTreeNodeItem> Node : RootNodeArray)
+	{
+		if (!ensure(Node))
+		{
+			continue;
+		}
+		RecursiveSetExpand(Node, false);
+	}
+	return FReply::Handled();
+}
+
+TSharedPtr<SWidget> SInterchangePipelineStacksTreeView::OnOpenContextMenu()
+{
+	// Build up the menu for a selection
+	const bool bCloseAfterSelection = true;
+	FMenuBuilder MenuBuilder(bCloseAfterSelection, TSharedPtr<FUICommandList>());
+
+	TArray<TSharedPtr<FInterchangePipelineStacksTreeNodeItem>> SelectedNodes;
+	const auto NumSelectedItems = GetSelectedItems(SelectedNodes);
+
+	if (SelectedNodes.Num() == 1)
+	{
+		TSharedPtr<FInterchangePipelineStacksTreeNodeItem> SelectNode = SelectedNodes[0];
+		if (SelectNode->Pipeline == nullptr)
+		{
+			// We always create a section here, even if there is no parent so that clients can still extend the menu
+			MenuBuilder.BeginSection("TreeViewContextMenuStackNodeSection");
+			{
+				const FSlateIcon DefaultIcon(FEditorStyle::GetStyleSetName(), "Icons.Default");
+				MenuBuilder.AddMenuEntry(LOCTEXT("SetHasDefaultMenuAction", "Set Has Default Stack"), FText(), DefaultIcon, FUIAction(FExecuteAction::CreateSP(this, &SInterchangePipelineStacksTreeView::SetHasDefaultStack, SelectNode->StackName)));
+			}
+			MenuBuilder.EndSection();
+		}
+	}
+	
+
+	return MenuBuilder.MakeWidget();
+}
+
+void SInterchangePipelineStacksTreeView::SetHasDefaultStack(FName NewDefaultStackValue)
+{
+	UInterchangeProjectSettings* InterchangeProjectSettingsCDO = GetMutableDefault<UInterchangeProjectSettings>();
+	FName& DefaultPipelineStackName = InterchangeProjectSettingsCDO->DefaultPipelineStack;
+	const TMap<FName, FInterchangePipelineStack>& PipelineStacks = InterchangeProjectSettingsCDO->PipelineStacks;
+	if (PipelineStacks.Contains(NewDefaultStackValue))
+	{
+		DefaultPipelineStackName = NewDefaultStackValue;
+		InterchangeProjectSettingsCDO->SaveConfig(); //This ensure the default pipeline stack name is save into the local config
+	}
+}
+
+void SInterchangePipelineStacksTreeView::OnTreeViewSelectionChanged(TSharedPtr<FInterchangePipelineStacksTreeNodeItem> Item, ESelectInfo::Type SelectionType)
+{
+	if (SelectionMode.Get() == ESelectionMode::None)
+	{
+		return;
+	}
+
+	if (OnSelectionChangedDelegate.IsBound())
+	{
+		OnSelectionChangedDelegate.ExecuteIfBound(Item, SelectionType);
+	}
+}
+
+
+/************************************************************************/
+/* SInterchangePipelineConfigurationDialog Implementation                      */
+/************************************************************************/
+
+SInterchangePipelineConfigurationDialog::SInterchangePipelineConfigurationDialog()
+{
+	PipelineConfigurationTreeView = nullptr;
+	PipelineConfigurationDetailsView = nullptr;
+	OwnerWindow = nullptr;
+}
+
+SInterchangePipelineConfigurationDialog::~SInterchangePipelineConfigurationDialog()
+{
+	PipelineConfigurationTreeView = nullptr;
+	PipelineConfigurationDetailsView = nullptr;
+	OwnerWindow = nullptr;
+}
+
+TSharedRef<SBox> SInterchangePipelineConfigurationDialog::SpawnPipelineConfiguration()
+{
+	//Create the treeview
+	PipelineConfigurationTreeView = SNew(SInterchangePipelineStacksTreeView)
+		.OnSelectionChangedDelegate(this, &SInterchangePipelineConfigurationDialog::OnSelectionChanged);
+
+	TSharedPtr<SBox> InspectorBox;
+	TSharedRef<SBox> PipelineConfigurationPanelBox = SNew(SBox)
+	[
+		SNew(SSplitter)
+		.Orientation(Orient_Horizontal)
+		+ SSplitter::Slot()
+		.Value(0.4f)
+		[
+			SNew(SVerticalBox)
+			+ SVerticalBox::Slot()
+			.HAlign(HAlign_Left)
+			.AutoHeight()
+			[
+				SNew(SUniformGridPanel)
+				.SlotPadding(2)
+				+ SUniformGridPanel::Slot(0, 0)
+				[
+					SNew(SButton)
+					.HAlign(HAlign_Center)
+					.Text(LOCTEXT("FbxOptionWindow_Scene_ExpandAll", "Expand All"))
+					.OnClicked(PipelineConfigurationTreeView.Get(), &SInterchangePipelineStacksTreeView::OnExpandAll)
+				]
+				+ SUniformGridPanel::Slot(1, 0)
+				[
+					SNew(SButton)
+					.HAlign(HAlign_Center)
+					.Text(LOCTEXT("FbxOptionWindow_Scene_CollapseAll", "Collapse All"))
+					.OnClicked(PipelineConfigurationTreeView.Get(), &SInterchangePipelineStacksTreeView::OnCollapseAll)
+				]
+			]
+			+ SVerticalBox::Slot()
+			.FillHeight(1.0f)
+			[
+				SNew(SBox)
+				[
+					PipelineConfigurationTreeView.ToSharedRef()
+				]
+			]
+		]
+		+ SSplitter::Slot()
+		.Value(0.6f)
+		[
+			SAssignNew(InspectorBox, SBox)
+		]
+	];
+
+	FPropertyEditorModule& PropertyEditorModule = FModuleManager::GetModuleChecked<FPropertyEditorModule>("PropertyEditor");
+	FDetailsViewArgs DetailsViewArgs;
+	DetailsViewArgs.bAllowSearch = false;
+	DetailsViewArgs.NameAreaSettings = FDetailsViewArgs::HideNameArea;
+	PipelineConfigurationDetailsView = PropertyEditorModule.CreateDetailView(DetailsViewArgs);
+	InspectorBox->SetContent(PipelineConfigurationDetailsView->AsShared());
+	PipelineConfigurationDetailsView->SetObject(nullptr);
+	return PipelineConfigurationPanelBox;
+}
+
+
+void SInterchangePipelineConfigurationDialog::Construct(const FArguments& InArgs)
+{
+	//Make sure there is a valid default value
+
+	OwnerWindow = InArgs._OwnerWindow;
+
+	check(OwnerWindow.IsValid());
+
+	this->ChildSlot
+	[
+		SNew(SBorder)
+		.Padding(FMargin(10.0f, 3.0f))
+		.BorderImage(FEditorStyle::GetBrush("ToolPanel.DarkGroupBorder"))
+		[
+			SNew(SVerticalBox)
+			+ SVerticalBox::Slot()
+			.FillHeight(1.0f)
+			.Padding(2)
+			[
+				SpawnPipelineConfiguration()
+			]
+			+ SVerticalBox::Slot()
+			.AutoHeight()
+			.HAlign(HAlign_Right)
+			.Padding(2)
+			[
+				SNew(SUniformGridPanel)
+				.SlotPadding(2)
+				+ SUniformGridPanel::Slot(0, 0)
+				[
+					IDocumentation::Get()->CreateAnchor(FString("Engine/Content/Interchange/PipelineConfiguration"))
+				]
+				+ SUniformGridPanel::Slot(1, 0)
+				[
+					SNew(SButton)
+					.HAlign(HAlign_Center)
+					.Text(LOCTEXT("InspectorGraphWindow_Cancel", "Cancel"))
+					.OnClicked(this, &SInterchangePipelineConfigurationDialog::OnCloseDialog, ECloseEventType::Cancel)
+				]
+				+ SUniformGridPanel::Slot(2, 0)
+				[
+					SNew(SButton)
+					.HAlign(HAlign_Center)
+					.Text(LOCTEXT("InspectorGraphWindow_ImportAll", "Import All"))
+					.OnClicked(this, &SInterchangePipelineConfigurationDialog::OnCloseDialog, ECloseEventType::ImportAll)
+				]
+				+ SUniformGridPanel::Slot(3, 0)
+				[
+					SNew(SButton)
+					.HAlign(HAlign_Center)
+					.Text(LOCTEXT("InspectorGraphWindow_Import", "Import"))
+					.OnClicked(this, &SInterchangePipelineConfigurationDialog::OnCloseDialog, ECloseEventType::Import)
+				]
+			]
+		]
+	];
+}
+
+void SInterchangePipelineConfigurationDialog::OnSelectionChanged(TSharedPtr<FInterchangePipelineStacksTreeNodeItem> Item, ESelectInfo::Type SelectionType)
+{
+	//Change the object point by the InspectorBox
+	PipelineConfigurationDetailsView->SetObject(Item->Pipeline);
+}
+
+void SInterchangePipelineConfigurationDialog::RecursiveSavePipelineSettings(const TSharedPtr<FInterchangePipelineStacksTreeNodeItem>& ParentNode, const int32 PipelineIndex) const
+{
+	if (ParentNode->Pipeline)
+	{
+		ParentNode->Pipeline->SaveSettings(ParentNode->StackName);
+	}
+	for (int32 ChildIndex = 0; ChildIndex < ParentNode->Childrens.Num(); ++ChildIndex)
+	{
+		RecursiveSavePipelineSettings(ParentNode->Childrens[ChildIndex], ChildIndex);
+	}
+}
+
+void SInterchangePipelineConfigurationDialog::ClosePipelineConfiguration(const ECloseEventType CloseEventType)
+{
+	if (CloseEventType == ECloseEventType::Cancel)
+	{
+		bCanceled = true;
+		bImportAll = false;
+	}
+	else if (CloseEventType == ECloseEventType::ImportAll)
+	{
+		bCanceled = false;
+		bImportAll = true;
+	}
+	else
+	{
+		//ECloseEventType::Import
+		bCanceled = false;
+		bImportAll = false;
+	}
+
+	if (PipelineConfigurationTreeView)
+	{
+		const TArray<TSharedPtr<FInterchangePipelineStacksTreeNodeItem>>& RootNodeArray = PipelineConfigurationTreeView->GetRootNodeArray();
+		for (const TSharedPtr<FInterchangePipelineStacksTreeNodeItem>& RootNode : RootNodeArray)
+		{
+			RecursiveSavePipelineSettings(RootNode, 0);
+		}
+	}
+
+	PipelineConfigurationTreeView = nullptr;
+	PipelineConfigurationDetailsView = nullptr;
+
+	if (TSharedPtr<SWindow> OwnerWindowPin = OwnerWindow.Pin())
+	{
+		OwnerWindowPin->RequestDestroyWindow();
+	}
+	OwnerWindow = nullptr;
+}
+
+FReply SInterchangePipelineConfigurationDialog::OnKeyDown(const FGeometry& MyGeometry, const FKeyEvent& InKeyEvent)
+{
+	if (InKeyEvent.GetKey() == EKeys::Escape)
+	{
+		if (!FApp::IsUnattended())
+		{
+			FString Message = FText(LOCTEXT("InterchangePipelineCancelEscKey", "Are you sure you want to cancel the import?")).ToString();
+			if (FPlatformMisc::MessageBoxExt(EAppMsgType::Ok, *Message, TEXT("Cancel Import")) == EAppReturnType::Type::Yes)
+			{
+				return OnCloseDialog(ECloseEventType::Cancel);
+			}
+		}
+	}
+	return FReply::Unhandled();
+}
+
+#undef LOCTEXT_NAMESPACE
diff --git a/Engine/Plugins/Editor/InterchangeEditor/Source/Pipelines/Private/SInterchangePipelineConfigurationDialog.h b/Engine/Plugins/Editor/InterchangeEditor/Source/Pipelines/Private/SInterchangePipelineConfigurationDialog.h
new file mode 100644
--- /dev/null
+++ b/Engine/Plugins/Editor/InterchangeEditor/Source/Pipelines/Private/SInterchangePipelineConfigurationDialog.h
@@ -0,0 +1,127 @@
+// Copyright Epic Games, Inc. All Rights Reserved.
+
+#pragma once
+
+#include "CoreMinimal.h"
+#include "InputCoreTypes.h"
+#include "Input/Reply.h"
+#include "InterchangePipelineBase.h"
+#include "Styling/SlateColor.h"
+#include "UObject/GCObject.h"
+#include "Widgets/DeclarativeSyntaxSupport.h"
+#include "Widgets/SCompoundWidget.h"
+#include "Widgets/Views/STreeView.h"
+
+class IDetailsView;
+
+class FInterchangePipelineStacksTreeNodeItem : protected FGCObject
+{
+public:
+	FInterchangePipelineStacksTreeNodeItem()
+	{
+		StackName = NAME_None;
+		Pipeline = nullptr;
+	}
+
+	/* FGCObject interface */
+	virtual void AddReferencedObjects(FReferenceCollector& Collector) override;
+
+	//This name is use only when this item represent a stack name
+	FName StackName;
+
+	//Pipeline is nullptr when the node represent a stack name
+	UInterchangePipelineBase* Pipeline;
+
+	TArray<TSharedPtr<FInterchangePipelineStacksTreeNodeItem>> Childrens;
+};
+
+DECLARE_DELEGATE_TwoParams(FOnPipelineConfigurationSelectionChanged, TSharedPtr<FInterchangePipelineStacksTreeNodeItem>, ESelectInfo::Type)
+
+class SInterchangePipelineStacksTreeView : public STreeView< TSharedPtr<FInterchangePipelineStacksTreeNodeItem>>
+{
+public:
+	~SInterchangePipelineStacksTreeView();
+
+	SLATE_BEGIN_ARGS(SInterchangePipelineStacksTreeView)
+		: _OnSelectionChangedDelegate()
+	{}
+		SLATE_EVENT(FOnPipelineConfigurationSelectionChanged, OnSelectionChangedDelegate)
+	SLATE_END_ARGS()
+
+	/** Construct this widget */
+	void Construct(const FArguments& InArgs);
+	TSharedRef< ITableRow > OnGenerateRowPipelineConfigurationTreeView(TSharedPtr<FInterchangePipelineStacksTreeNodeItem> Item, const TSharedRef< STableViewBase >& OwnerTable);
+	void OnGetChildrenPipelineConfigurationTreeView(TSharedPtr<FInterchangePipelineStacksTreeNodeItem> InParent, TArray< TSharedPtr<FInterchangePipelineStacksTreeNodeItem> >& OutChildren);
+
+	FReply OnExpandAll();
+	FReply OnCollapseAll();
+
+	const TArray<TSharedPtr<FInterchangePipelineStacksTreeNodeItem>>& GetRootNodeArray() const { return RootNodeArray; }
+protected:
+	/** Delegate to invoke when selection changes. */
+	FOnPipelineConfigurationSelectionChanged OnSelectionChangedDelegate;
+
+	/** the elements we show in the tree view */
+	TArray<TSharedPtr<FInterchangePipelineStacksTreeNodeItem>> RootNodeArray;
+
+	/** Open a context menu for the current selection */
+	TSharedPtr<SWidget> OnOpenContextMenu();
+	void SetHasDefaultStack(FName NewDefaultStackValue);
+	void RecursiveSetExpand(TSharedPtr<FInterchangePipelineStacksTreeNodeItem> Node, bool ExpandState);
+	void OnTreeViewSelectionChanged(TSharedPtr<FInterchangePipelineStacksTreeNodeItem> Item, ESelectInfo::Type SelectionType);
+};
+
+
+enum class ECloseEventType : uint8
+{
+	Cancel,
+	ImportAll,
+	Import
+};
+
+class SInterchangePipelineConfigurationDialog : public SCompoundWidget
+{
+public:
+	SLATE_BEGIN_ARGS(SInterchangePipelineConfigurationDialog)
+		: _OwnerWindow()
+	{}
+
+		SLATE_ARGUMENT(TWeakPtr<SWindow>, OwnerWindow)
+	SLATE_END_ARGS()
+
+public:
+
+	SInterchangePipelineConfigurationDialog();
+	~SInterchangePipelineConfigurationDialog();
+
+	void Construct(const FArguments& InArgs);
+	virtual bool SupportsKeyboardFocus() const override { return true; }
+
+	void ClosePipelineConfiguration(const ECloseEventType CloseEventType);
+
+	FReply OnCloseDialog(const ECloseEventType CloseEventType)
+	{
+		ClosePipelineConfiguration(CloseEventType);
+		return FReply::Handled();
+	}
+
+	virtual FReply OnKeyDown(const FGeometry& MyGeometry, const FKeyEvent& InKeyEvent) override;
+
+	bool IsCanceled() { return bCanceled; }
+	bool IsImportAll() { return bImportAll; }
+private:
+	TWeakPtr< SWindow > OwnerWindow;
+
+	TSharedRef<SBox> SpawnPipelineConfiguration();
+	void OnSelectionChanged(TSharedPtr<FInterchangePipelineStacksTreeNodeItem> Item, ESelectInfo::Type SelectionType);
+
+	void RecursiveSavePipelineSettings(const TSharedPtr<FInterchangePipelineStacksTreeNodeItem>& ParentNode, const int32 PipelineIndex) const;
+
+	//Graph Inspector UI elements
+	TSharedPtr<SInterchangePipelineStacksTreeView> PipelineConfigurationTreeView;
+	TSharedPtr<IDetailsView> PipelineConfigurationDetailsView;
+
+	bool bCanceled = false;
+	bool bImportAll = false;
+};
+
diff --git a/Engine/Plugins/Editor/InterchangeEditor/Source/Pipelines/Public/InterchangePipelineConfigurationGeneric.h b/Engine/Plugins/Editor/InterchangeEditor/Source/Pipelines/Public/InterchangePipelineConfigurationGeneric.h
new file mode 100644
--- /dev/null
+++ b/Engine/Plugins/Editor/InterchangeEditor/Source/Pipelines/Public/InterchangePipelineConfigurationGeneric.h
@@ -0,0 +1,22 @@
+// Copyright Epic Games, Inc. All Rights Reserved.
+
+#pragma once
+
+#include "CoreMinimal.h"
+#include "InterchangePipelineConfigurationBase.h"
+#include "UObject/Object.h"
+#include "UObject/ObjectMacros.h"
+
+#include "InterchangePipelineConfigurationGeneric.generated.h"
+
+UCLASS(BlueprintType, Blueprintable)
+class INTERCHANGEEDITORPIPELINES_API UInterchangePipelineConfigurationGeneric : public UInterchangePipelineConfigurationBase
+{
+	GENERATED_BODY()
+
+public:
+
+protected:
+
+	virtual EInterchangePipelineConfigurationDialogResult ShowPipelineConfigurationDialog() override;
+};
diff --git a/Engine/Plugins/Runtime/Interchange/Source/Import/Private/InterchangeImportCommon.cpp b/Engine/Plugins/Runtime/Interchange/Source/Import/Private/InterchangeImportCommon.cpp
--- a/Engine/Plugins/Runtime/Interchange/Source/Import/Private/InterchangeImportCommon.cpp
+++ b/Engine/Plugins/Runtime/Interchange/Source/Import/Private/InterchangeImportCommon.cpp
@@ -4,6 +4,7 @@
 #include "CoreMinimal.h"
 #include "EditorFramework/AssetImportData.h"
 #include "InterchangeAssetImportData.h"
+#include "InterchangePipelineBase.h"
 #include "InterchangeSourceData.h"
 #include "Nodes/InterchangeBaseNodeContainer.h"
 #include "UObject/Object.h"
@@ -16,12 +17,14 @@ namespace UE
 																							, UAssetImportData* InAssetImportData
 																							, const UInterchangeSourceData* InSourceData
 																							, FString InNodeUniqueID
-																							, UInterchangeBaseNodeContainer* InNodeContainer)
+																							, UInterchangeBaseNodeContainer* InNodeContainer
+																						   , const TArray<UInterchangePipelineBase*>& InPipelines)
 			: AssetImportDataOuter(InAssetImportDataOuter)
 			, AssetImportData(InAssetImportData)
 			, SourceData(InSourceData)
 			, NodeUniqueID(InNodeUniqueID)
 			, NodeContainer(InNodeContainer)
+			, Pipelines(InPipelines)
 		{
 			ensure(AssetImportDataOuter);
 			ensure(SourceData);
@@ -83,7 +86,14 @@ namespace UE
 			AssetImportData->NodeUniqueID = Parameters.NodeUniqueID;
 			FObjectDuplicationParameters DupParam(Parameters.NodeContainer, AssetImportData);
 			AssetImportData->NodeContainer = CastChecked<UInterchangeBaseNodeContainer>(StaticDuplicateObjectEx(DupParam));
-
+			for (const UInterchangePipelineBase* Pipeline : Parameters.Pipelines)
+			{
+				UInterchangePipelineBase* DupPipeline = Cast<UInterchangePipelineBase>(StaticDuplicateObject(Pipeline, AssetImportData));
+				if (DupPipeline)
+				{
+					AssetImportData->Pipelines.Add(DupPipeline);
+				}
+			}
 			// Return the asset import data so it can be set on the imported asset.
 			return AssetImportData;
 #endif //#if WITH_EDITORONLY_DATA
diff --git a/Engine/Plugins/Runtime/Interchange/Source/Import/Private/InterchangeImportCommon.h b/Engine/Plugins/Runtime/Interchange/Source/Import/Private/InterchangeImportCommon.h
--- a/Engine/Plugins/Runtime/Interchange/Source/Import/Private/InterchangeImportCommon.h
+++ b/Engine/Plugins/Runtime/Interchange/Source/Import/Private/InterchangeImportCommon.h
@@ -7,6 +7,7 @@
 class UAssetImportData;
 class UInterchangeBaseNode;
 class UInterchangeBaseNodeContainer;
+class UInterchangePipelineBase;
 class UInterchangeSourceData;
 class UObject;
 
@@ -29,12 +30,14 @@ namespace UE
 				const UInterchangeSourceData* SourceData = nullptr;
 				FString NodeUniqueID;
 				UInterchangeBaseNodeContainer* NodeContainer;
+				const TArray<UInterchangePipelineBase*> Pipelines;
 
 				FUpdateImportAssetDataParameters(UObject* InAssetImportDataOuter
 																	, UAssetImportData* InAssetImportData
 																	, const UInterchangeSourceData* InSourceData
 																	, FString InNodeUniqueID
-																	, UInterchangeBaseNodeContainer* InNodeContainer);
+																	, UInterchangeBaseNodeContainer* InNodeContainer
+																	, const TArray<UInterchangePipelineBase*>& InPipelines);
 			};
 
 			/**
diff --git a/Engine/Plugins/Runtime/Interchange/Source/Import/Private/Material/InterchangeMaterialFactory.cpp b/Engine/Plugins/Runtime/Interchange/Source/Import/Private/Material/InterchangeMaterialFactory.cpp
--- a/Engine/Plugins/Runtime/Interchange/Source/Import/Private/Material/InterchangeMaterialFactory.cpp
+++ b/Engine/Plugins/Runtime/Interchange/Source/Import/Private/Material/InterchangeMaterialFactory.cpp
@@ -318,7 +318,8 @@ void UInterchangeMaterialFactory::PostImportGameThreadCallback(const FPostImport
 																										  , ImportedMaterial->AssetImportData
 																										  , Arguments.SourceData
 																										  , Arguments.NodeUniqueID
-																										  , Arguments.NodeContainer);
+																										  , Arguments.NodeContainer
+																										  , Arguments.Pipelines);
 
 		ImportedMaterial->AssetImportData = UE::Interchange::FFactoryCommon::UpdateImportAssetData(UpdateImportAssetDataParameters);
 	}
diff --git a/Engine/Plugins/Runtime/Interchange/Source/Import/Private/Mesh/InterchangeSkeletalMeshFactory.cpp b/Engine/Plugins/Runtime/Interchange/Source/Import/Private/Mesh/InterchangeSkeletalMeshFactory.cpp
--- a/Engine/Plugins/Runtime/Interchange/Source/Import/Private/Mesh/InterchangeSkeletalMeshFactory.cpp
+++ b/Engine/Plugins/Runtime/Interchange/Source/Import/Private/Mesh/InterchangeSkeletalMeshFactory.cpp
@@ -235,11 +235,11 @@ namespace UE
 							//We did not found any material in the dependencies so try to find material everywhere
 							Material = FindObject<UMaterialInterface>(ANY_PACKAGE, *ImportedMaterial.MaterialImportName);
 						}
+
+						const bool bEnableShadowCasting = true;
+						const bool bInRecomputeTangent = false;
+						Materials.Add(FSkeletalMaterial(Material, bEnableShadowCasting, bInRecomputeTangent, Material != nullptr ? Material->GetFName() : FName(*ImportedMaterial.MaterialImportName), FName(*(ImportedMaterial.MaterialImportName))));
 					}
-					
-					const bool bEnableShadowCasting = true;
-					const bool bInRecomputeTangent = false;
-					Materials.Add(FSkeletalMaterial(Material, bEnableShadowCasting, bInRecomputeTangent, Material != nullptr ? Material->GetFName() : FName(*ImportedMaterial.MaterialImportName), FName(*(ImportedMaterial.MaterialImportName))));
 				}
 
 				int32 NumMaterialsToAdd = FMath::Max<int32>(ImportedMaterials.Num(), ImportData.MaxMaterialIndex + 1);
@@ -853,7 +853,8 @@ void UInterchangeSkeletalMeshFactory::PostImportGameThreadCallback(const FPostIm
 																										  , ImportDataPtr
 																										  , Arguments.SourceData
 																										  , Arguments.NodeUniqueID
-																										  , Arguments.NodeContainer);
+																										  , Arguments.NodeContainer
+																										  , Arguments.Pipelines);
 
 		ImportDataPtr = UE::Interchange::FFactoryCommon::UpdateImportAssetData(UpdateImportAssetDataParameters);
 		SkeletalMesh->SetAssetImportData(ImportDataPtr);
diff --git a/Engine/Plugins/Runtime/Interchange/Source/Import/Private/Texture/InterchangeTextureFactory.cpp b/Engine/Plugins/Runtime/Interchange/Source/Import/Private/Texture/InterchangeTextureFactory.cpp
--- a/Engine/Plugins/Runtime/Interchange/Source/Import/Private/Texture/InterchangeTextureFactory.cpp
+++ b/Engine/Plugins/Runtime/Interchange/Source/Import/Private/Texture/InterchangeTextureFactory.cpp
@@ -562,7 +562,8 @@ void UInterchangeTextureFactory::PostImportGameThreadCallback(const FPostImportG
 																						 , ImportedTexture->AssetImportData
 																						 , Arguments.SourceData
 																						 , Arguments.NodeUniqueID
-																						 , Arguments.NodeContainer);
+																						 , Arguments.NodeContainer
+																						 , Arguments.Pipelines);
 		ImportedTexture->AssetImportData = UE::Interchange::FFactoryCommon::UpdateImportAssetData(UpdateImportAssetDataParameters);
 	}
 #endif
diff --git a/Engine/Plugins/Runtime/Interchange/Source/Pipelines/Private/InterchangeGenericAssetsPipeline.cpp b/Engine/Plugins/Runtime/Interchange/Source/Pipelines/Private/InterchangeGenericAssetsPipeline.cpp
--- a/Engine/Plugins/Runtime/Interchange/Source/Pipelines/Private/InterchangeGenericAssetsPipeline.cpp
+++ b/Engine/Plugins/Runtime/Interchange/Source/Pipelines/Private/InterchangeGenericAssetsPipeline.cpp
@@ -570,6 +570,11 @@ void UInterchangeGenericAssetsPipeline::AddLodDataToSkeletalMesh(const UIntercha
 	for (const TPair<int32, TArray<FString>>& LodIndexAndNodeUids : NodeUidsPerLodIndex)
 	{
 		const int32 LodIndex = LodIndexAndNodeUids.Key;
+		if (!bImportLods && LodIndex > 0)
+		{
+			//If the pipeline should not import lods, skip any lod over base lod
+			continue;
+		}
 		const TArray<FString>& NodeUids = LodIndexAndNodeUids.Value;
 
 		//Create a lod data node with all the meshes for this LOD
diff --git a/Engine/Source/Runtime/Interchange/Core/Private/InterchangePipelineBase.cpp b/Engine/Source/Runtime/Interchange/Core/Private/InterchangePipelineBase.cpp
new file mode 100644
--- /dev/null
+++ b/Engine/Source/Runtime/Interchange/Core/Private/InterchangePipelineBase.cpp
@@ -0,0 +1,173 @@
+// Copyright Epic Games, Inc. All Rights Reserved.
+
+#include "InterchangePipelineBase.h"
+
+#include "CoreMinimal.h"
+#include "InterchangeLogPrivate.h"
+#include "Misc/ConfigCacheIni.h"
+#include "UObject/Object.h"
+#include "UObject/ObjectMacros.h"
+
+namespace UE
+{
+	namespace Interchange
+	{
+		namespace PipelinePrivate
+		{
+			FString CreateConfigSectionName(const FName PipelineStackName, UClass* PipelineClass)
+			{
+				FString Section = TEXT("Interchange__StackName_") + PipelineStackName.ToString() + TEXT("__PipelineClassName_") + PipelineClass->GetName();
+				return Section;
+			}
+		}
+	}
+}
+
+void UInterchangePipelineBase::LoadSettings(const FName PipelineStackName)
+{
+	int32 PortFlags = 0;
+	UClass* Class = this->GetClass();
+	for (FProperty* Property = Class->PropertyLink; Property; Property = Property->PropertyLinkNext)
+	{
+// 		if (!Property->HasAnyPropertyFlags(CPF_Config))
+// 		{
+// 			continue;
+// 		}
+		FString Section = UE::Interchange::PipelinePrivate::CreateConfigSectionName(PipelineStackName, Class);
+		FString Key = Property->GetName();
+
+		const bool bIsPropertyInherited = Property->GetOwnerClass() != Class;
+		UObject* SuperClassDefaultObject = Class->GetSuperClass()->GetDefaultObject();
+
+		const FString& PropFileName = GEditorPerProjectIni;
+
+		FArrayProperty* Array = CastField<FArrayProperty>(Property);
+		if (Array)
+		{
+			const bool bForce = false;
+			const bool bConst = true;
+			FConfigSection* Sec = GConfig->GetSectionPrivate(*Section, bForce, bConst, *GEditorPerProjectIni);
+			if (Sec != nullptr)
+			{
+				TArray<FConfigValue> List;
+				const FName KeyName(*Key, FNAME_Find);
+				Sec->MultiFind(KeyName, List);
+
+				FScriptArrayHelper_InContainer ArrayHelper(Array, this);
+				// Only override default properties if there is something to override them with.
+				if (List.Num() > 0)
+				{
+					ArrayHelper.EmptyAndAddValues(List.Num());
+					for (int32 i = List.Num() - 1, c = 0; i >= 0; i--, c++)
+					{
+						Array->Inner->ImportText(*List[i].GetValue(), ArrayHelper.GetRawPtr(c), PortFlags, this);
+					}
+				}
+				else
+				{
+					int32 Index = 0;
+					const FConfigValue* ElementValue = nullptr;
+					do
+					{
+						// Add array index number to end of key
+						FString IndexedKey = FString::Printf(TEXT("%s[%i]"), *Key, Index);
+
+						// Try to find value of key
+						const FName IndexedName(*IndexedKey, FNAME_Find);
+						if (IndexedName == NAME_None)
+						{
+							break;
+						}
+						ElementValue = Sec->Find(IndexedName);
+
+						// If found, import the element
+						if (ElementValue != nullptr)
+						{
+							// expand the array if necessary so that Index is a valid element
+							ArrayHelper.ExpandForIndex(Index);
+							Array->Inner->ImportText(*ElementValue->GetValue(), ArrayHelper.GetRawPtr(Index), PortFlags, this);
+						}
+
+						Index++;
+					} while (ElementValue || Index < ArrayHelper.Num());
+				}
+			}
+		}
+		else
+		{
+			for (int32 i = 0; i < Property->ArrayDim; i++)
+			{
+				if (Property->ArrayDim != 1)
+				{
+					Key = FString::Printf(TEXT("%s[%i]"), *Property->GetName(), i);
+				}
+
+				FString Value;
+				bool bFoundValue = GConfig->GetString(*Section, *Key, Value, *GEditorPerProjectIni);
+
+				if (bFoundValue)
+				{
+					if (Property->ImportText(*Value, Property->ContainerPtrToValuePtr<uint8>(this, i), PortFlags, this) == NULL)
+					{
+						// this should be an error as the properties from the .ini / .int file are not correctly being read in and probably are affecting things in subtle ways
+						UE_LOG(LogInterchangeCore, Error, TEXT("UInterchangePipeline (class:%s) failed to load settings. Property: %s Value: %s"), *this->GetClass()->GetName(), *Property->GetName(), *Value);
+					}
+				}
+			}
+		}
+	}
+}
+
+void UInterchangePipelineBase::SaveSettings(const FName PipelineStackName)
+{
+	int32 PortFlags = 0;
+	UClass* Class = this->GetClass();
+	for (FProperty* Property = Class->PropertyLink; Property; Property = Property->PropertyLinkNext)
+	{
+// 		if (!Property->HasAnyPropertyFlags(CPF_Config))
+// 		{
+// 			continue;
+// 		}
+		FString Section = UE::Interchange::PipelinePrivate::CreateConfigSectionName(PipelineStackName, Class);
+		FString Key = Property->GetName();
+
+		const bool bIsPropertyInherited = Property->GetOwnerClass() != Class;
+		UObject* SuperClassDefaultObject = Class->GetSuperClass()->GetDefaultObject();
+
+		FArrayProperty* Array = CastField<FArrayProperty>(Property);
+		if (Array)
+		{
+			const bool bForce = true;
+			const bool bConst = false;
+			FConfigSection* Sec = GConfig->GetSectionPrivate(*Section, bForce, bConst, *GEditorPerProjectIni);
+			check(Sec);
+			Sec->Remove(*Key);
+
+			FScriptArrayHelper_InContainer ArrayHelper(Array, this);
+			for (int32 i = 0; i < ArrayHelper.Num(); i++)
+			{
+				FString	Buffer;
+				Array->Inner->ExportTextItem(Buffer, ArrayHelper.GetRawPtr(i), ArrayHelper.GetRawPtr(i), this, PortFlags);
+				Sec->Add(*Key, *Buffer);
+			}
+		}
+		else
+		{
+			TCHAR TempKey[MAX_SPRINTF] = TEXT("");
+			for (int32 Index = 0; Index < Property->ArrayDim; Index++)
+			{
+				if (Property->ArrayDim != 1)
+				{
+					FCString::Sprintf(TempKey, TEXT("%s[%i]"), *Property->GetName(), Index);
+					Key = TempKey;
+				}
+
+				FString	Value;
+				Property->ExportText_InContainer(Index, Value, this, this, this, PortFlags);
+				GConfig->SetString(*Section, *Key, *Value, *GEditorPerProjectIni);
+			}
+		}
+	}
+	GConfig->Flush(0);
+}
+
diff --git a/Engine/Source/Runtime/Interchange/Core/Public/InterchangeFactoryBase.h b/Engine/Source/Runtime/Interchange/Core/Public/InterchangeFactoryBase.h
--- a/Engine/Source/Runtime/Interchange/Core/Public/InterchangeFactoryBase.h
+++ b/Engine/Source/Runtime/Interchange/Core/Public/InterchangeFactoryBase.h
@@ -11,6 +11,7 @@
 
 class UInterchangeBaseNode;
 class UInterchangeBaseNodeContainer;
+class UInterchangePipelineBase;
 class UInterchangeSourceData;
 class UInterchangeTranslatorBase;
 
@@ -116,6 +117,7 @@ public:
 		UObject* ImportedObject = nullptr;
 		FString NodeUniqueID;
 		UInterchangeBaseNodeContainer* NodeContainer = nullptr;
+		TArray<UInterchangePipelineBase*> Pipelines;
 ;
 	};
 	/* This function is call in the completion task on the main thread, use it to call main thread post creation step for your assets*/
diff --git a/Engine/Source/Runtime/Interchange/Core/Public/InterchangePipelineBase.h b/Engine/Source/Runtime/Interchange/Core/Public/InterchangePipelineBase.h
--- a/Engine/Source/Runtime/Interchange/Core/Public/InterchangePipelineBase.h
+++ b/Engine/Source/Runtime/Interchange/Core/Public/InterchangePipelineBase.h
@@ -89,6 +89,10 @@ public:
 		return true;
 	}
 
+	void LoadSettings(const FName PipelineStackName);
+
+	void SaveSettings(const FName PipelineStackName);
+
 protected:
 
 	/**
diff --git a/Engine/Source/Runtime/Interchange/Engine/Private/InterchangeManager.cpp b/Engine/Source/Runtime/Interchange/Engine/Private/InterchangeManager.cpp
--- a/Engine/Source/Runtime/Interchange/Engine/Private/InterchangeManager.cpp
+++ b/Engine/Source/Runtime/Interchange/Engine/Private/InterchangeManager.cpp
@@ -7,6 +7,7 @@
 #include "CoreMinimal.h"
 #include "Engine/Blueprint.h"
 #include "Framework/Notifications/NotificationManager.h"
+#include "InterchangeAssetImportData.h"
 #include "InterchangeFactoryBase.h"
 #include "InterchangeEngineLogPrivate.h"
 #include "InterchangeProjectSettings.h"
@@ -547,7 +548,7 @@ UE::Interchange::FAssetImportResultRef UInterchangeManager::ImportAssetAsync(con
 		//Import process can be started only in the game thread
 		return MakeShared< UE::Interchange::FAssetImportResult, ESPMode::ThreadSafe >();
 	}
-
+	UInterchangeAssetImportData* OriginalAssetImportData = nullptr;
 	FString PackageBasePath = ContentPath;
 	if(!ImportAssetParameters.ReimportAsset)
 	{
@@ -556,8 +557,18 @@ UE::Interchange::FAssetImportResultRef UInterchangeManager::ImportAssetAsync(con
 	else
 	{
 		PackageBasePath = FPaths::GetPath(ImportAssetParameters.ReimportAsset->GetPathName());
+		TArray<UObject*> SubObjects;
+		ImportAssetParameters.ReimportAsset->CollectDefaultSubobjects(SubObjects, false);
+		for (UObject* SubObject : SubObjects)
+		{
+			OriginalAssetImportData = Cast<UInterchangeAssetImportData>(SubObject);
+			if(OriginalAssetImportData)
+			{
+				break;
+			}
+		}
 	}
-
+	bool bImportCancel = false;
 	bool bCanShowDialog = !ImportAssetParameters.bIsAutomated && !IsAttended();
 
 	//Create a task for every source data
@@ -586,19 +597,119 @@ UE::Interchange::FAssetImportResultRef UInterchangeManager::ImportAssetAsync(con
 		check(AsyncHelper->BaseNodeContainers[SourceDataIndex].IsValid());
 	}
 
+	UInterchangePipelineConfigurationBase* RegisteredPipelineConfiguration = nullptr;
+	TSoftClassPtr <UInterchangePipelineConfigurationBase> PipelineConfigurationDialogClass = GetDefault<UInterchangeProjectSettings>()->PipelineConfigurationDialogClass;
+	if (PipelineConfigurationDialogClass.IsValid())
+	{
+		UClass* PipelineConfigurationClass = PipelineConfigurationDialogClass.LoadSynchronous();
+		if (PipelineConfigurationClass)
+		{
+			RegisteredPipelineConfiguration = NewObject<UInterchangePipelineConfigurationBase>(GetTransientPackage(), PipelineConfigurationClass, NAME_None, RF_NoFlags);
+		}
+	}
+
+
 	if ( ImportAssetParameters.OverridePipelines.Num() == 0 )
 	{
-		const TArray<TSoftClassPtr<UInterchangePipelineBase>>& PipelineStack = GetDefault<UInterchangeProjectSettings>()->PipelineStack;
+		const bool bIsUnattended = FApp::IsUnattended() || GIsAutomationTesting;
+
+		const bool bShowPipelineStacksConfigurationDialog = !bIsUnattended
+															&& GetDefault<UInterchangeProjectSettings>()->bShowPipelineStacksConfigurationDialog
+															&& !bImportAllWithDefault;
 
-		for (int32 GraphPipelineIndex = 0; GraphPipelineIndex < PipelineStack.Num(); ++GraphPipelineIndex)
+		const FName DefaultPipelineStackName = GetDefault<UInterchangeProjectSettings>()->DefaultPipelineStack;
+		const TMap<FName, FInterchangePipelineStack>& DefaultPipelineStacks = GetDefault<UInterchangeProjectSettings>()->PipelineStacks;
+		
+		//If we reimport we want to load the original pipeline and the original pipeline settings
+		if (OriginalAssetImportData && OriginalAssetImportData->Pipelines.Num() > 0)
 		{
-			if (PipelineStack[GraphPipelineIndex].IsValid())
+			TArray<UInterchangePipelineBase*> PipelineStack;
+			for (int32 PipelineIndex = 0; PipelineIndex < OriginalAssetImportData->Pipelines.Num(); ++PipelineIndex)
 			{
-				UClass* PipelineClass = PipelineStack[GraphPipelineIndex].LoadSynchronous();
-				if (PipelineClass)
+				UInterchangePipelineBase* SourcePipeline = OriginalAssetImportData->Pipelines[PipelineIndex];
+				if (SourcePipeline) //Its possible a pipeline doesnt exist anymore so it wont load into memory when we loading the outer asset
+				{
+					//Duplicate the pipeline saved in the asset import data
+					UInterchangePipelineBase* GeneratedPipeline = Cast<UInterchangePipelineBase>(StaticDuplicateObject(SourcePipeline, GetTransientPackage()));
+					PipelineStack.Add(GeneratedPipeline);
+				}
+				else
 				{
-					UInterchangePipelineBase* GeneratedPipeline = NewObject<UInterchangePipelineBase>(GetTransientPackage(), PipelineClass, NAME_None, RF_NoFlags);
-					AsyncHelper->Pipelines.Add(GeneratedPipeline);
+					//A pipeline was not loaded
+					//Log something
+				}
+			}
+
+			if (RegisteredPipelineConfiguration && (bShowPipelineStacksConfigurationDialog || (!DefaultPipelineStacks.Contains(DefaultPipelineStackName) && !bIsUnattended)))
+			{
+				//Show the re-import dialog to let the user make change in the pipelines
+				//PipelineStackName = RegisteredPipelineConfiguration->ScriptedShowReimportPipelineConfigurationDialog(PipelineStack);
+			}
+
+			//If the Stack name is empty it mean we want to use the re-import stack (PipelineStack). If there is a name we will
+			//extract the pipeline stack the user want to use.
+		}
+		else
+		{
+			FName PipelineStackName = DefaultPipelineStackName;
+			if (RegisteredPipelineConfiguration && (bShowPipelineStacksConfigurationDialog || (!DefaultPipelineStacks.Contains(PipelineStackName) && !bIsUnattended)))
+			{
+				//Show the dialog, a plugin should have register this dialog. We use a plugin to be able to use editor code when doing UI
+				EInterchangePipelineConfigurationDialogResult DialogResult = RegisteredPipelineConfiguration->ScriptedShowPipelineConfigurationDialog();
+				if (DialogResult == EInterchangePipelineConfigurationDialogResult::Cancel)
+				{
+					bImportCancel = true;
+				}
+				if (DialogResult == EInterchangePipelineConfigurationDialogResult::ImportAll)
+				{
+					bImportAllWithDefault = true;
+				}
+				PipelineStackName = GetDefault<UInterchangeProjectSettings>()->DefaultPipelineStack;
+			}
+			if (!bImportCancel)
+			{
+				//Get the latest PipelineStacks (the Dialog can change the CDO)
+				const TMap<FName, FInterchangePipelineStack>& PipelineStacks = GetDefault<UInterchangeProjectSettings>()->PipelineStacks;
+				if (!PipelineStacks.Contains(PipelineStackName))
+				{
+					if (PipelineStacks.Contains(DefaultPipelineStackName))
+					{
+						PipelineStackName = DefaultPipelineStackName;
+					}
+					else
+					{
+						//Log an error, we cannot import asset without a valid pipeline, we will use the first pipeline
+						for (const TPair<FName, FInterchangePipelineStack>& PipelineStack : PipelineStacks)
+						{
+							PipelineStackName = PipelineStack.Key;
+						}
+					}
+				}
+
+				if (PipelineStacks.Contains(PipelineStackName))
+				{
+					//use the default pipeline
+					const FInterchangePipelineStack& PipelineStack = PipelineStacks.FindChecked(PipelineStackName);
+					TArray<TSoftClassPtr<UInterchangePipelineBase>> Pipelines = PipelineStack.Pipelines;
+					for (int32 GraphPipelineIndex = 0; GraphPipelineIndex < Pipelines.Num(); ++GraphPipelineIndex)
+					{
+						if (Pipelines[GraphPipelineIndex].IsValid())
+						{
+							UClass* PipelineClass = Pipelines[GraphPipelineIndex].LoadSynchronous();
+							if (PipelineClass)
+							{
+								UInterchangePipelineBase* GeneratedPipeline = NewObject<UInterchangePipelineBase>(GetTransientPackage(), PipelineClass, NAME_None, RF_NoFlags);
+								//Load the settings for this pipeline
+								GeneratedPipeline->LoadSettings(PipelineStackName);
+								AsyncHelper->Pipelines.Add(GeneratedPipeline);
+							}
+						}
+					}
+				}
+				else
+				{
+					//Log an error, we cannot import asset without a valid pipeline, there is no pipeline stack defined
+					bImportCancel = true;
 				}
 			}
 		}
@@ -612,6 +723,14 @@ UE::Interchange::FAssetImportResultRef UInterchangeManager::ImportAssetAsync(con
 		}
 	}
 
+	//Cancel the import do not queue task
+	if (bImportCancel)
+	{
+		AsyncHelper->InitCancel();
+		AsyncHelper->CleanUp();
+	}
+
+	//Queue the task cancel or not, we need to return a valid asset import result
 	FQueuedTaskData QueuedTaskData;
 	QueuedTaskData.AsyncHelper = AsyncHelper;
 	QueuedTaskData.PackageBasePath = PackageBasePath;
@@ -673,6 +792,7 @@ void UInterchangeManager::ReleaseAsyncHelper(TWeakPtr<UE::Interchange::FImportAs
 	FString ImportTaskNumberStr = TEXT(" (") + FString::FromInt(ImportTaskNumber) + TEXT(")");
 	if (ImportTaskNumber == 0)
 	{
+		bImportAllWithDefault = false;
 		SetActiveMode(false);
 
 		if (Notification.IsValid())
diff --git a/Engine/Source/Runtime/Interchange/Engine/Private/Tasks/InterchangeTaskCompletion.cpp b/Engine/Source/Runtime/Interchange/Engine/Private/Tasks/InterchangeTaskCompletion.cpp
--- a/Engine/Source/Runtime/Interchange/Engine/Private/Tasks/InterchangeTaskCompletion.cpp
+++ b/Engine/Source/Runtime/Interchange/Engine/Private/Tasks/InterchangeTaskCompletion.cpp
@@ -56,6 +56,7 @@ void UE::Interchange::FTaskCompletion::DoTask(ENamedThreads::Type CurrentThread,
 				Arguments.SourceData = AsyncHelper->SourceDatas[SourceIndex];
 				Arguments.NodeUniqueID = AssetInfo.NodeUniqueId;
 				Arguments.NodeContainer = AsyncHelper->BaseNodeContainers[SourceIndex].Get();
+				Arguments.Pipelines = AsyncHelper->Pipelines;
 				AssetInfo.Factory->PostImportGameThreadCallback(Arguments);
 			}
 
diff --git a/Engine/Source/Runtime/Interchange/Engine/Public/InterchangeAssetImportData.h b/Engine/Source/Runtime/Interchange/Engine/Public/InterchangeAssetImportData.h
--- a/Engine/Source/Runtime/Interchange/Engine/Public/InterchangeAssetImportData.h
+++ b/Engine/Source/Runtime/Interchange/Engine/Public/InterchangeAssetImportData.h
@@ -10,9 +10,10 @@
 #include "UObject/Object.h"
 #include "UObject/ObjectMacros.h"
 
-
 #include "InterchangeAssetImportData.generated.h"
 
+class UInterchangePipelineBase;
+
 UCLASS(BlueprintType)
 class INTERCHANGEENGINE_API UInterchangeAssetImportData : public UAssetImportData
 {
@@ -68,4 +69,8 @@ public:
 	/** The Node UID pass to the factory that exist in the graph that was use to create this asset */
 	UPROPERTY(VisibleAnywhere, Category = "Interchange | AssetImportData")
 	FString NodeUniqueID;
+
+	/** List of pipelines use to import an asset */
+	UPROPERTY(VisibleAnywhere, Category = "Interchange | AssetImportData")
+	TArray<TObjectPtr<UInterchangePipelineBase>> Pipelines;
 };
diff --git a/Engine/Source/Runtime/Interchange/Engine/Public/InterchangeManager.h b/Engine/Source/Runtime/Interchange/Engine/Public/InterchangeManager.h
--- a/Engine/Source/Runtime/Interchange/Engine/Public/InterchangeManager.h
+++ b/Engine/Source/Runtime/Interchange/Engine/Public/InterchangeManager.h
@@ -11,6 +11,7 @@
 #include "HAL/ThreadSafeBool.h"
 #include "InterchangeFactoryBase.h"
 #include "InterchangePipelineBase.h"
+#include "InterchangePipelineConfigurationBase.h"
 #include "InterchangeSourceData.h"
 #include "InterchangeTranslatorBase.h"
 #include "InterchangeWriterBase.h"
@@ -464,5 +465,10 @@ private:
 	//If interchange is currently importing we have a timer to watch the cancel and we block GC 
 	FThreadSafeBool bIsActive = false;
 	FThread GcGuardThread;
+
+	//If the user want to use the same import pipeline stack for all the queue task
+	//This boolean is reset to false when the ImportTasks array is empty.
+	bool bImportAllWithDefault = false;
+
 	friend class UE::Interchange::FScopedTranslator;
 };
diff --git a/Engine/Source/Runtime/Interchange/Engine/Public/InterchangePipelineConfigurationBase.h b/Engine/Source/Runtime/Interchange/Engine/Public/InterchangePipelineConfigurationBase.h
new file mode 100644
--- /dev/null
+++ b/Engine/Source/Runtime/Interchange/Engine/Public/InterchangePipelineConfigurationBase.h
@@ -0,0 +1,48 @@
+// Copyright Epic Games, Inc. All Rights Reserved.
+
+#pragma once
+
+#include "CoreMinimal.h"
+#include "UObject/Object.h"
+#include "UObject/ObjectMacros.h"
+
+#include "InterchangePipelineConfigurationBase.generated.h"
+
+UENUM(BlueprintType)
+enum class EInterchangePipelineConfigurationDialogResult : uint8
+{
+	Cancel		UMETA(DisplayName = "Cancel"),
+	Import		UMETA(DisplayName = "Import"),
+	ImportAll	UMETA(DisplayName = "Import All"),
+};
+
+UCLASS(BlueprintType, Blueprintable)
+class INTERCHANGEENGINE_API UInterchangePipelineConfigurationBase : public UObject
+{
+	GENERATED_BODY()
+
+public:
+
+	/**
+	 * Non virtual helper to allow blueprint to implement event base function to implement the ShowPipelineConfigurationDialog,
+	 */
+	UFUNCTION(BlueprintNativeEvent, BlueprintCallable, Category = "Interchange | Translator")
+	EInterchangePipelineConfigurationDialogResult ScriptedShowPipelineConfigurationDialog();
+	/** The default implementation (call if the blueprint do not have any implementation) will call the virtual ExecuteImportPipeline */
+	EInterchangePipelineConfigurationDialogResult ScriptedShowPipelineConfigurationDialog_Implementation()
+	{
+		//By default we call the virtual import pipeline execution
+		return ShowPipelineConfigurationDialog();
+	}
+
+protected:
+
+	/**
+	 * This function show a dialog use to configure pipeline stacks and return a stack name that tell the caller the user choice.
+	 */
+	virtual EInterchangePipelineConfigurationDialogResult ShowPipelineConfigurationDialog()
+	{ 
+		//Not implemented
+		return EInterchangePipelineConfigurationDialogResult::Cancel;
+	}
+};
diff --git a/Engine/Source/Runtime/Interchange/Engine/Public/InterchangeProjectSettings.h b/Engine/Source/Runtime/Interchange/Engine/Public/InterchangeProjectSettings.h
--- a/Engine/Source/Runtime/Interchange/Engine/Public/InterchangeProjectSettings.h
+++ b/Engine/Source/Runtime/Interchange/Engine/Public/InterchangeProjectSettings.h
@@ -6,16 +6,39 @@
 #include "Engine/DeveloperSettings.h"
 #include "Engine/EngineTypes.h"
 #include "InterchangePipelineBase.h"
+#include "InterchangePipelineConfigurationBase.h"
 
 #include "InterchangeProjectSettings.generated.h"
 
+USTRUCT()
+struct FInterchangePipelineStack
+{
+	GENERATED_BODY()
+
+	/** The starting mesh for the blueprint **/
+	UPROPERTY(EditAnywhere, config, Category = Interchange)
+	TArray<TSoftClassPtr<UInterchangePipelineBase>> Pipelines;
+};
+
 UCLASS(config=Engine, meta=(DisplayName=Interchange), MinimalAPI)
 class UInterchangeProjectSettings : public UDeveloperSettings
 {
 	GENERATED_BODY()
 
 public:
-	/** Stack the pipeline you want to use to import asset with interchange. The pipelines are execute from top to bottom order. You can order them by using the grip on the left of any pipelines.*/
-	UPROPERTY(EditAnywhere, config, Category = Interchange, meta = (DisplayName = "Import pipeline stack"))
-	TArray<TSoftClassPtr<UInterchangePipelineBase>> PipelineStack;
+	/** All the available pipeline stacks you want to use to import with interchange. The chosen pipeline stack execute all the pipelines from top to bottom order. You can order them by using the grip on the left of any pipelines.*/
+	UPROPERTY(EditAnywhere, config, Category = Interchange)
+	TMap<FName, FInterchangePipelineStack> PipelineStacks;
+
+	/** This tell interchange which pipeline to select when importing.*/
+	UPROPERTY(EditAnywhere, config, Category = Interchange)
+	FName DefaultPipelineStack;
+
+	/** This tell interchange which pipeline configuration dialog to popup when we need to configure the pipelines.*/
+	UPROPERTY(EditAnywhere, config, Category = Interchange)
+	TSoftClassPtr <UInterchangePipelineConfigurationBase> PipelineConfigurationDialogClass;
+
+	/** If enabled, the pipeline stacks configuration dialog will show every time interchange must choose a pipeline to import or re-import. If disabled interchange will use the DefaultPipelineStack.*/
+	UPROPERTY(EditAnywhere, config, Category = Interchange)
+	bool bShowPipelineStacksConfigurationDialog;
 };
\ No newline at end of file
