Add Interchange vertex color option to sk and sm

diff --git a/Engine/Plugins/Editor/InterchangeEditor/Source/Pipelines/Private/InterchangeEditorPipelineDetails.cpp b/Engine/Plugins/Editor/InterchangeEditor/Source/Pipelines/Private/InterchangeEditorPipelineDetails.cpp
--- a/Engine/Plugins/Editor/InterchangeEditor/Source/Pipelines/Private/InterchangeEditorPipelineDetails.cpp
+++ b/Engine/Plugins/Editor/InterchangeEditor/Source/Pipelines/Private/InterchangeEditorPipelineDetails.cpp
@@ -190,6 +190,18 @@ void FInterchangeBaseNodeDetailsCustomization::AddAttributeRow(UE::Interchange::
 		}
 		break;
 
+		case UE::Interchange::EAttributeTypes::Color:
+		{
+			BuildColorValueContent(AttributeCategory, AttributeKey);
+		}
+		break;
+
+		case UE::Interchange::EAttributeTypes::LinearColor:
+		{
+			BuildLinearColorValueContent(AttributeCategory, AttributeKey);
+		}
+		break;
+
 		default:
 		{
 			FText AttributeName = FText::FromString(InterchangeBaseNode->GetKeyDisplayName(AttributeKey));
@@ -505,6 +517,236 @@ void FInterchangeBaseNodeDetailsCustomization::BuildTransformValueContent(IDetai
     ];
 }
 
+void FInterchangeBaseNodeDetailsCustomization::BuildColorValueContent(IDetailCategoryBuilder& AttributeCategory, UE::Interchange::FAttributeKey& AttributeKey)
+{
+	UE::Interchange::EAttributeTypes AttributeType = InterchangeBaseNode->GetAttributeType(AttributeKey);
+	if (AttributeType == UE::Interchange::EAttributeTypes::Color)
+	{
+		const UE::Interchange::FAttributeStorage::TAttributeHandle<FColor> AttributeHandle = InterchangeBaseNode->GetAttributeHandle<FColor>(AttributeKey);
+		if (!AttributeHandle.IsValid())
+		{
+			CreateInvalidHandleRow(AttributeCategory, AttributeKey);
+			return;
+		}
+	}
+	else
+	{
+		ensure(AttributeType == UE::Interchange::EAttributeTypes::Color);
+		CreateInvalidHandleRow(AttributeCategory, AttributeKey);
+		return;
+	}
+	InternalBuildColorValueContent<FColor, uint8>(AttributeCategory, AttributeKey, 255);
+}
+
+void FInterchangeBaseNodeDetailsCustomization::BuildLinearColorValueContent(IDetailCategoryBuilder& AttributeCategory, UE::Interchange::FAttributeKey& AttributeKey)
+{
+	UE::Interchange::EAttributeTypes AttributeType = InterchangeBaseNode->GetAttributeType(AttributeKey);
+	if (AttributeType == UE::Interchange::EAttributeTypes::LinearColor)
+	{
+		const UE::Interchange::FAttributeStorage::TAttributeHandle<FColor> AttributeHandle = InterchangeBaseNode->GetAttributeHandle<FColor>(AttributeKey);
+		if (!AttributeHandle.IsValid())
+		{
+			CreateInvalidHandleRow(AttributeCategory, AttributeKey);
+			return;
+		}
+	}
+	else
+	{
+		ensure(AttributeType == UE::Interchange::EAttributeTypes::LinearColor);
+		CreateInvalidHandleRow(AttributeCategory, AttributeKey);
+		return;
+	}
+	InternalBuildColorValueContent<FLinearColor, float>(AttributeCategory, AttributeKey, 1.0f);
+}
+
+template<typename AttributeType, typename NumericType>
+void FInterchangeBaseNodeDetailsCustomization::InternalBuildColorValueContent(IDetailCategoryBuilder& AttributeCategory, UE::Interchange::FAttributeKey& AttributeKey, NumericType DefaultTypeValue)
+{
+	const bool bAdvancedProperty = false;
+	const FString GroupName = InterchangeBaseNode->GetKeyDisplayName(AttributeKey);
+	IDetailGroup& Group = AttributeCategory.AddGroup(FName(*GroupName), FText::FromString(GroupName), bAdvancedProperty);
+	FDetailWidgetRow& GroupHeaderRow = Group.HeaderRow();
+	GroupHeaderRow.NameContent().Widget = SNew(SBox)
+	[
+		CreateNameWidget(AttributeKey)
+	];
+
+	auto GetRValue = [DefaultTypeValue](int32 ComponentIndex, UInterchangeBaseNode* BaseNode, UE::Interchange::FAttributeKey& Key)->NumericType
+	{
+		AttributeType Value;
+		const UE::Interchange::FAttributeStorage::TAttributeHandle<AttributeType> AttributeHandle = BaseNode->GetAttributeHandle<AttributeType>(Key);
+		if (AttributeHandle.IsValid())
+		{
+			AttributeHandle.Get(Value);
+		}
+		else
+		{
+			//Error return white color
+			return DefaultTypeValue;
+		}
+		return Value.R;
+	};
+
+	auto SetRValue = [](int32 ComponentIndex, UInterchangeBaseNode* BaseNode, const NumericType& Value, UE::Interchange::FAttributeKey& Key)
+	{
+		AttributeType ColorValue;
+		UE::Interchange::FAttributeStorage::TAttributeHandle<AttributeType> AttributeHandle = BaseNode->GetAttributeHandle<AttributeType>(Key);
+		if (AttributeHandle.IsValid())
+		{
+			AttributeHandle.Get(ColorValue);
+			ColorValue.R = Value;
+			AttributeHandle.Set(ColorValue);
+		}
+	};
+
+	auto GetGValue = [DefaultTypeValue](int32 ComponentIndex, UInterchangeBaseNode* BaseNode, UE::Interchange::FAttributeKey& Key)->NumericType
+	{
+		AttributeType Value;
+		const UE::Interchange::FAttributeStorage::TAttributeHandle<AttributeType> AttributeHandle = BaseNode->GetAttributeHandle<AttributeType>(Key);
+		if (AttributeHandle.IsValid())
+		{
+			AttributeHandle.Get(Value);
+		}
+		else
+		{
+			//Error return white color
+			return DefaultTypeValue;
+		}
+		return Value.G;
+	};
+
+	auto SetGValue = [](int32 ComponentIndex, UInterchangeBaseNode* BaseNode, const NumericType& Value, UE::Interchange::FAttributeKey& Key)
+	{
+		AttributeType ColorValue;
+		UE::Interchange::FAttributeStorage::TAttributeHandle<AttributeType> AttributeHandle = BaseNode->GetAttributeHandle<AttributeType>(Key);
+		if (AttributeHandle.IsValid())
+		{
+			AttributeHandle.Get(ColorValue);
+			ColorValue.G = Value;
+			AttributeHandle.Set(ColorValue);
+		}
+	};
+
+	auto GetBValue = [DefaultTypeValue](int32 ComponentIndex, UInterchangeBaseNode* BaseNode, UE::Interchange::FAttributeKey& Key)->NumericType
+	{
+		AttributeType Value;
+		const UE::Interchange::FAttributeStorage::TAttributeHandle<AttributeType> AttributeHandle = BaseNode->GetAttributeHandle<AttributeType>(Key);
+		if (AttributeHandle.IsValid())
+		{
+			AttributeHandle.Get(Value);
+		}
+		else
+		{
+			//Error return white color
+			return DefaultTypeValue;
+		}
+		return Value.B;
+	};
+
+	auto SetBValue = [](int32 ComponentIndex, UInterchangeBaseNode* BaseNode, const NumericType& Value, UE::Interchange::FAttributeKey& Key)
+	{
+		AttributeType ColorValue;
+		UE::Interchange::FAttributeStorage::TAttributeHandle<AttributeType> AttributeHandle = BaseNode->GetAttributeHandle<AttributeType>(Key);
+		if (AttributeHandle.IsValid())
+		{
+			AttributeHandle.Get(ColorValue);
+			ColorValue.B = Value;
+			AttributeHandle.Set(ColorValue);
+		}
+	};
+
+	auto GetAValue = [DefaultTypeValue](int32 ComponentIndex, UInterchangeBaseNode* BaseNode, UE::Interchange::FAttributeKey& Key)->NumericType
+	{
+		AttributeType Value;
+		const UE::Interchange::FAttributeStorage::TAttributeHandle<AttributeType> AttributeHandle = BaseNode->GetAttributeHandle<AttributeType>(Key);
+		if (AttributeHandle.IsValid())
+		{
+			AttributeHandle.Get(Value);
+		}
+		else
+		{
+			//Error return white color
+			return DefaultTypeValue;
+		}
+		return Value.A;
+	};
+
+	auto SetAValue = [](int32 ComponentIndex, UInterchangeBaseNode* BaseNode, const NumericType& Value, UE::Interchange::FAttributeKey& Key)
+	{
+		AttributeType ColorValue;
+		UE::Interchange::FAttributeStorage::TAttributeHandle<AttributeType> AttributeHandle = BaseNode->GetAttributeHandle<AttributeType>(Key);
+		if (AttributeHandle.IsValid())
+		{
+			AttributeHandle.Get(ColorValue);
+			ColorValue.A = Value;
+			AttributeHandle.Set(ColorValue);
+		}
+	};
+
+	const FString RName = TEXT("Red");
+	Group.AddWidgetRow()
+    .NameContent()
+    [
+        CreateSimpleNameWidget(RName)
+    ]
+    .ValueContent()
+    [
+		SNew(SHorizontalBox)
+		+ SHorizontalBox::Slot()
+		.AutoWidth()
+		[
+			MakeNumericWidget<NumericType>(0, GetRValue, SetRValue, AttributeKey)
+		]
+    ];
+	const FString GName = TEXT("Green");
+	Group.AddWidgetRow()
+    .NameContent()
+    [
+        CreateSimpleNameWidget(GName)
+    ]
+    .ValueContent()
+    [
+		SNew(SHorizontalBox)
+		+ SHorizontalBox::Slot()
+		.AutoWidth()
+		[
+			MakeNumericWidget<NumericType>(0, GetGValue, SetGValue, AttributeKey)
+		]
+    ];
+
+	const FString BName = TEXT("Blue");
+	Group.AddWidgetRow()
+    .NameContent()
+    [
+        CreateSimpleNameWidget(BName)
+    ]
+    .ValueContent()
+    [
+		SNew(SHorizontalBox)
+		+ SHorizontalBox::Slot()
+		.AutoWidth()
+		[
+			MakeNumericWidget<NumericType>(0, GetBValue, SetBValue, AttributeKey)
+		]
+    ];
+
+	const FString AName = TEXT("Alpha");
+	Group.AddWidgetRow()
+    .NameContent()
+    [
+        CreateSimpleNameWidget(AName)
+    ]
+    .ValueContent()
+    [
+		SNew(SHorizontalBox)
+		+ SHorizontalBox::Slot()
+		.AutoWidth()
+		[
+			MakeNumericWidget<NumericType>(0, GetAValue, SetAValue, AttributeKey)
+		]
+    ];
+}
+
 void FInterchangeBaseNodeDetailsCustomization::BuildBoxValueContent(IDetailCategoryBuilder& AttributeCategory, UE::Interchange::FAttributeKey& AttributeKey)
 {
 	UE::Interchange::EAttributeTypes AttributeType = InterchangeBaseNode->GetAttributeType(AttributeKey);
@@ -777,4 +1019,5 @@ TSharedRef<SWidget> FInterchangeBaseNodeDetailsCustomization::MakeNumericWidget(
 			.OnValueChanged_Lambda(SetValueChangedLambda)
 			.AllowSpin(false);
 }
+
 #undef LOCTEXT_NAMESPACE
diff --git a/Engine/Plugins/Editor/InterchangeEditor/Source/Pipelines/Private/InterchangeEditorPipelineDetails.h b/Engine/Plugins/Editor/InterchangeEditor/Source/Pipelines/Private/InterchangeEditorPipelineDetails.h
--- a/Engine/Plugins/Editor/InterchangeEditor/Source/Pipelines/Private/InterchangeEditorPipelineDetails.h
+++ b/Engine/Plugins/Editor/InterchangeEditor/Source/Pipelines/Private/InterchangeEditorPipelineDetails.h
@@ -58,6 +58,17 @@ private:
 	//Any numeric field are read/write
 	void BuildTransformValueContent(IDetailCategoryBuilder& AttributeCategory, UE::Interchange::FAttributeKey& AttributeKey);
 
+	//FColor: Add a row to the category that contain the RGBA value
+	//Any numeric field are read/write
+	void BuildColorValueContent(IDetailCategoryBuilder& AttributeCategory, UE::Interchange::FAttributeKey& AttributeKey);
+
+	//FLinearColor: Add a row to the category that contain the RGBA value
+	//Any numeric field are read/write
+	void BuildLinearColorValueContent(IDetailCategoryBuilder& AttributeCategory, UE::Interchange::FAttributeKey& AttributeKey);
+
+	template<typename AttributeType, typename NumericType>
+	void InternalBuildColorValueContent(IDetailCategoryBuilder& AttributeCategory, UE::Interchange::FAttributeKey& AttributeKey, NumericType DefaultTypeValue);
+
 	//FBox: Add a group to the category that contain the Minimum corner position(FVector), and the maximum corner position(FVector)
 	//Any numeric field are read/write
 	void BuildBoxValueContent(IDetailCategoryBuilder& AttributeCategory, UE::Interchange::FAttributeKey& AttributeKey);
diff --git a/Engine/Plugins/Runtime/Interchange/Source/Import/Private/Mesh/InterchangeSkeletalMeshFactory.cpp b/Engine/Plugins/Runtime/Interchange/Source/Import/Private/Mesh/InterchangeSkeletalMeshFactory.cpp
--- a/Engine/Plugins/Runtime/Interchange/Source/Import/Private/Mesh/InterchangeSkeletalMeshFactory.cpp
+++ b/Engine/Plugins/Runtime/Interchange/Source/Import/Private/Mesh/InterchangeSkeletalMeshFactory.cpp
@@ -920,6 +920,44 @@ UObject* UInterchangeSkeletalMeshFactory::CreateAsset(const UInterchangeSkeletal
 																						   , RefBonesBinary
 																						   , Arguments
 																						   , SkeletalMeshTranslatorPayloadInterface);
+				//////////////////////////////////////////////////////////////////////////
+				//Manage vertex color, we want to use the translated source data
+				//Replace -> do nothing
+				//Ignore -> remove vertex color from import data (when we re-import, ignore have to put back the current mesh vertex color)
+				//Override -> replace the vertex color by the override color
+				{
+					bool bReplaceVertexColor = false;
+					SkeletalMeshFactoryNode->GetCustomVertexColorReplace(bReplaceVertexColor);
+					if (!bReplaceVertexColor)
+					{
+						bool bIgnoreVertexColor = false;
+						SkeletalMeshFactoryNode->GetCustomVertexColorIgnore(bIgnoreVertexColor);
+						if (bIgnoreVertexColor)
+						{
+							//Flush the vertex color, if we re-import we have to fill it with the old data
+							SkeletalMeshImportData.bHasVertexColors = false;
+							for (SkeletalMeshImportData::FVertex& Wedge : SkeletalMeshImportData.Wedges)
+							{
+								Wedge.Color = FColor::White;
+							}
+						}
+						else
+						{
+							FColor OverrideVertexColor;
+							if (SkeletalMeshFactoryNode->GetCustomVertexColorOverride(OverrideVertexColor))
+							{
+								SkeletalMeshImportData.bHasVertexColors = true;
+								for(SkeletalMeshImportData::FVertex& Wedge : SkeletalMeshImportData.Wedges)
+								{
+									Wedge.Color = OverrideVertexColor;
+								}
+							}
+						}
+					}
+					// Store whether or not this mesh has vertex colors
+					SkeletalMesh->SetHasVertexColors(SkeletalMeshImportData.bHasVertexColors);
+					SkeletalMesh->SetVertexColorGuid(SkeletalMesh->GetHasVertexColors() ? FGuid::NewGuid() : FGuid());
+				}
 
 				ensure(ImportedResource->LODModels.Add(new FSkeletalMeshLODModel()) == CurrentLodIndex);
 				FSkeletalMeshLODModel& LODModel = ImportedResource->LODModels[CurrentLodIndex];
diff --git a/Engine/Plugins/Runtime/Interchange/Source/Import/Private/Mesh/InterchangeStaticMeshFactory.cpp b/Engine/Plugins/Runtime/Interchange/Source/Import/Private/Mesh/InterchangeStaticMeshFactory.cpp
--- a/Engine/Plugins/Runtime/Interchange/Source/Import/Private/Mesh/InterchangeStaticMeshFactory.cpp
+++ b/Engine/Plugins/Runtime/Interchange/Source/Import/Private/Mesh/InterchangeStaticMeshFactory.cpp
@@ -268,6 +268,42 @@ UObject* UInterchangeStaticMeshFactory::CreateAsset(const UInterchangeStaticMesh
 					FStaticMeshOperations::AppendMeshDescription(LodMeshPayload->MeshDescription, *LodMeshDescription, AppendSettings);
 				}
 
+				//////////////////////////////////////////////////////////////////////////
+				//Manage vertex color
+				//Replace -> do nothing, we want to use the translated source data
+				//Ignore -> remove vertex color from import data (when we re-import, ignore have to put back the current mesh vertex color)
+				//Override -> replace the vertex color by the override color
+				{
+					FStaticMeshAttributes Attributes(*LodMeshDescription);
+					TVertexInstanceAttributesRef<FVector4> VertexInstanceColors = Attributes.GetVertexInstanceColors();
+					bool bReplaceVertexColor = false;
+					StaticMeshFactoryNode->GetCustomVertexColorReplace(bReplaceVertexColor);
+					if (!bReplaceVertexColor)
+					{
+						bool bIgnoreVertexColor = false;
+						StaticMeshFactoryNode->GetCustomVertexColorIgnore(bIgnoreVertexColor);
+						if (bIgnoreVertexColor)
+						{
+							//Flush the vertex color, if we re-import we have to fill it with the old data
+							for (const FVertexInstanceID& VertexInstanceID : LodMeshDescription->VertexInstances().GetElementIDs())
+							{
+								VertexInstanceColors[VertexInstanceID] = FVector4(FLinearColor(FColor::White));
+							}
+						}
+						else
+						{
+							FColor OverrideVertexColor;
+							if (StaticMeshFactoryNode->GetCustomVertexColorOverride(OverrideVertexColor))
+							{
+								for (const FVertexInstanceID& VertexInstanceID : LodMeshDescription->VertexInstances().GetElementIDs())
+								{
+									VertexInstanceColors[VertexInstanceID] = FVector4(FLinearColor(OverrideVertexColor));
+								}
+							}
+						}
+					}
+				}
+
 				StaticMesh->CommitMeshDescription(CurrentLodIndex);
 
 				// Handle materials
diff --git a/Engine/Plugins/Runtime/Interchange/Source/Nodes/Public/InterchangeSkeletalMeshFactoryNode.h b/Engine/Plugins/Runtime/Interchange/Source/Nodes/Public/InterchangeSkeletalMeshFactoryNode.h
--- a/Engine/Plugins/Runtime/Interchange/Source/Nodes/Public/InterchangeSkeletalMeshFactoryNode.h
+++ b/Engine/Plugins/Runtime/Interchange/Source/Nodes/Public/InterchangeSkeletalMeshFactoryNode.h
@@ -199,6 +199,45 @@ public:
 	{
 		IMPLEMENT_NODE_ATTRIBUTE_SETTER_NODELEGATE(PhysicAssetSoftObjectPath, FSoftObjectPath)
 	}
+
+	/** Return false if the Attribute was not set previously.*/
+	UFUNCTION(BlueprintCallable, Category = "Interchange | Node | SkeletalMesh")
+	bool GetCustomVertexColorReplace(bool& AttributeValue) const
+	{
+		IMPLEMENT_NODE_ATTRIBUTE_GETTER(VertexColorReplace, bool)
+	}
+
+	UFUNCTION(BlueprintCallable, Category = "Interchange | Node | SkeletalMesh")
+	bool SetCustomVertexColorReplace(const bool& AttributeValue)
+	{
+		IMPLEMENT_NODE_ATTRIBUTE_SETTER_NODELEGATE(VertexColorReplace, bool)
+	}
+
+	/** Return false if the Attribute was not set previously.*/
+	UFUNCTION(BlueprintCallable, Category = "Interchange | Node | SkeletalMesh")
+	bool GetCustomVertexColorIgnore(bool& AttributeValue) const
+	{
+		IMPLEMENT_NODE_ATTRIBUTE_GETTER(VertexColorIgnore, bool)
+	}
+
+	UFUNCTION(BlueprintCallable, Category = "Interchange | Node | SkeletalMesh")
+	bool SetCustomVertexColorIgnore(const bool& AttributeValue)
+	{
+		IMPLEMENT_NODE_ATTRIBUTE_SETTER_NODELEGATE(VertexColorIgnore, bool)
+	}
+
+	/** Return false if the Attribute was not set previously.*/
+	UFUNCTION(BlueprintCallable, Category = "Interchange | Node | SkeletalMesh")
+	bool GetCustomVertexColorOverride(FColor& AttributeValue) const
+	{
+		IMPLEMENT_NODE_ATTRIBUTE_GETTER(VertexColorOverride, FColor)
+	}
+
+	UFUNCTION(BlueprintCallable, Category = "Interchange | Node | SkeletalMesh")
+	bool SetCustomVertexColorOverride(const FColor& AttributeValue)
+	{
+		IMPLEMENT_NODE_ATTRIBUTE_SETTER_NODELEGATE(VertexColorOverride, FColor)
+	}
 private:
 
 	void FillAssetClassFromAttribute()
@@ -238,6 +277,9 @@ private:
 	const UE::Interchange::FAttributeKey Macro_CustomSkeletonSoftObjectPathKey = UE::Interchange::FAttributeKey(TEXT("SkeletonSoftObjectPath"));
 	const UE::Interchange::FAttributeKey Macro_CustomCreatePhysicsAssetKey = UE::Interchange::FAttributeKey(TEXT("CreatePhysicsAsset"));
 	const UE::Interchange::FAttributeKey Macro_CustomPhysicAssetSoftObjectPathKey = UE::Interchange::FAttributeKey(TEXT("PhysicAssetSoftObjectPath"));
+	const UE::Interchange::FAttributeKey Macro_CustomVertexColorReplaceKey = UE::Interchange::FAttributeKey(TEXT("VertexColorReplace"));
+	const UE::Interchange::FAttributeKey Macro_CustomVertexColorIgnoreKey = UE::Interchange::FAttributeKey(TEXT("VertexColorIgnore"));
+	const UE::Interchange::FAttributeKey Macro_CustomVertexColorOverrideKey = UE::Interchange::FAttributeKey(TEXT("VertexColorOverride"));
 
 	UE::Interchange::FNameAttributeArrayHelper LodDependencies;
 protected:
diff --git a/Engine/Plugins/Runtime/Interchange/Source/Nodes/Public/InterchangeStaticMeshFactoryNode.h b/Engine/Plugins/Runtime/Interchange/Source/Nodes/Public/InterchangeStaticMeshFactoryNode.h
--- a/Engine/Plugins/Runtime/Interchange/Source/Nodes/Public/InterchangeStaticMeshFactoryNode.h
+++ b/Engine/Plugins/Runtime/Interchange/Source/Nodes/Public/InterchangeStaticMeshFactoryNode.h
@@ -146,6 +146,44 @@ public:
 		return LodDependencies.RemoveName(LodDataUniqueId);
 	}
 
+	/** Return false if the Attribute was not set previously.*/
+	UFUNCTION(BlueprintCallable, Category = "Interchange | Node | StaticMesh")
+	bool GetCustomVertexColorReplace(bool& AttributeValue) const
+	{
+		IMPLEMENT_NODE_ATTRIBUTE_GETTER(VertexColorReplace, bool)
+	}
+
+	UFUNCTION(BlueprintCallable, Category = "Interchange | Node | StaticMesh")
+	bool SetCustomVertexColorReplace(const bool& AttributeValue)
+	{
+		IMPLEMENT_NODE_ATTRIBUTE_SETTER_NODELEGATE(VertexColorReplace, bool)
+	}
+
+	/** Return false if the Attribute was not set previously.*/
+	UFUNCTION(BlueprintCallable, Category = "Interchange | Node | StaticMesh")
+	bool GetCustomVertexColorIgnore(bool& AttributeValue) const
+	{
+		IMPLEMENT_NODE_ATTRIBUTE_GETTER(VertexColorIgnore, bool)
+	}
+
+	UFUNCTION(BlueprintCallable, Category = "Interchange | Node | StaticMesh")
+	bool SetCustomVertexColorIgnore(const bool& AttributeValue)
+	{
+		IMPLEMENT_NODE_ATTRIBUTE_SETTER_NODELEGATE(VertexColorIgnore, bool)
+	}
+
+	/** Return false if the Attribute was not set previously.*/
+	UFUNCTION(BlueprintCallable, Category = "Interchange | Node | StaticMesh")
+	bool GetCustomVertexColorOverride(FColor& AttributeValue) const
+	{
+		IMPLEMENT_NODE_ATTRIBUTE_GETTER(VertexColorOverride, FColor)
+	}
+
+	UFUNCTION(BlueprintCallable, Category = "Interchange | Node | StaticMesh")
+	bool SetCustomVertexColorOverride(const FColor& AttributeValue)
+	{
+		IMPLEMENT_NODE_ATTRIBUTE_SETTER_NODELEGATE(VertexColorOverride, FColor)
+	}
 private:
 
 	void FillAssetClassFromAttribute()
@@ -181,6 +219,9 @@ private:
 	}
 
 	const UE::Interchange::FAttributeKey ClassNameAttributeKey = UE::Interchange::FBaseNodeStaticData::ClassTypeAttributeKey();
+	const UE::Interchange::FAttributeKey Macro_CustomVertexColorReplaceKey = UE::Interchange::FAttributeKey(TEXT("VertexColorReplace"));
+	const UE::Interchange::FAttributeKey Macro_CustomVertexColorIgnoreKey = UE::Interchange::FAttributeKey(TEXT("VertexColorIgnore"));
+	const UE::Interchange::FAttributeKey Macro_CustomVertexColorOverrideKey = UE::Interchange::FAttributeKey(TEXT("VertexColorOverride"));
 
 	UE::Interchange::FNameAttributeArrayHelper LodDependencies;
 protected:
diff --git a/Engine/Plugins/Runtime/Interchange/Source/Pipelines/Private/InterchangeGenericAssetsPipeline.cpp b/Engine/Plugins/Runtime/Interchange/Source/Pipelines/Private/InterchangeGenericAssetsPipeline.cpp
--- a/Engine/Plugins/Runtime/Interchange/Source/Pipelines/Private/InterchangeGenericAssetsPipeline.cpp
+++ b/Engine/Plugins/Runtime/Interchange/Source/Pipelines/Private/InterchangeGenericAssetsPipeline.cpp
@@ -507,6 +507,24 @@ UInterchangeStaticMeshFactoryNode* UInterchangeGenericAssetsPipeline::CreateStat
 
 	AddLodDataToStaticMesh(StaticMeshFactoryNode, MeshUidsPerLodIndex);
 
+	switch (VertexColorImportOption)
+	{
+		case EInterchangeVertexColorImportOption::IVCIO_Replace:
+		{
+			StaticMeshFactoryNode->SetCustomVertexColorReplace(true);
+		}
+		break;
+		case EInterchangeVertexColorImportOption::IVCIO_Ignore:
+		{
+			StaticMeshFactoryNode->SetCustomVertexColorIgnore(true);
+		}
+		break;
+		case EInterchangeVertexColorImportOption::IVCIO_Override:
+		{
+			StaticMeshFactoryNode->SetCustomVertexColorOverride(VertexOverrideColor);
+		}
+		break;
+	}
 	return StaticMeshFactoryNode;
 }
 
diff --git a/Engine/Plugins/Runtime/Interchange/Source/Pipelines/Private/InterchangeGenericSkeletalMeshPipeline.cpp b/Engine/Plugins/Runtime/Interchange/Source/Pipelines/Private/InterchangeGenericSkeletalMeshPipeline.cpp
--- a/Engine/Plugins/Runtime/Interchange/Source/Pipelines/Private/InterchangeGenericSkeletalMeshPipeline.cpp
+++ b/Engine/Plugins/Runtime/Interchange/Source/Pipelines/Private/InterchangeGenericSkeletalMeshPipeline.cpp
@@ -606,6 +606,26 @@ UInterchangeSkeletalMeshFactoryNode* UInterchangeGenericAssetsPipeline::CreateSk
 	}
 #endif
 
+	const bool bTrueValue = true;
+	switch (VertexColorImportOption)
+	{
+		case EInterchangeVertexColorImportOption::IVCIO_Replace:
+		{
+			SkeletalMeshFactoryNode->SetCustomVertexColorReplace(bTrueValue);
+		}
+		break;
+		case EInterchangeVertexColorImportOption::IVCIO_Ignore:
+		{
+			SkeletalMeshFactoryNode->SetCustomVertexColorIgnore(bTrueValue);
+		}
+		break;
+		case EInterchangeVertexColorImportOption::IVCIO_Override:
+		{
+			SkeletalMeshFactoryNode->SetCustomVertexColorOverride(VertexOverrideColor);
+		}
+		break;
+	}
+
 	return SkeletalMeshFactoryNode;
 }
 
diff --git a/Engine/Plugins/Runtime/Interchange/Source/Pipelines/Public/InterchangeGenericAssetsPipeline.h b/Engine/Plugins/Runtime/Interchange/Source/Pipelines/Public/InterchangeGenericAssetsPipeline.h
--- a/Engine/Plugins/Runtime/Interchange/Source/Pipelines/Public/InterchangeGenericAssetsPipeline.h
+++ b/Engine/Plugins/Runtime/Interchange/Source/Pipelines/Public/InterchangeGenericAssetsPipeline.h
@@ -41,16 +41,29 @@ template<class T> class TSubclassOf;
 
 /** Force mesh type, if user want to import all meshes as one type*/
 UENUM(BlueprintType)
-enum EForceMeshType
+enum EInterchangeForceMeshType
 {
 	/** Will import from the source type, no conversion */
-	FMT_None UMETA(DisplayName = "None"),
+	IFMT_None UMETA(DisplayName = "None"),
 	/** Will import any mesh to static mesh. */
-	FMT_StaticMesh UMETA(DisplayName = "Static Mesh"),
+	IFMT_StaticMesh UMETA(DisplayName = "Static Mesh"),
 	/** Will import any mesh to skeletal mesh. */
-	FMT_SkeletalMesh UMETA(DisplayName = "Skeletal Mesh"),
+	IFMT_SkeletalMesh UMETA(DisplayName = "Skeletal Mesh"),
 
-	FMT_MAX,
+	IFMT_MAX
+};
+
+UENUM(BlueprintType)
+enum EInterchangeVertexColorImportOption
+{
+	/** Import the mesh using the vertex colors from the translated source. */
+	IVCIO_Replace UMETA(DisplayName = "Replace"),
+	/** Ignore vertex colors from the translated source. In case of a re-import keep the existing mesh vertex colors. */
+	IVCIO_Ignore UMETA(DisplayName = "Ignore"),
+	/** Override all vertex colors with the specified color. */
+	IVCIO_Override UMETA(DisplayName = "Override"),
+
+	IVCIO_MAX
 };
 
 /**
@@ -91,6 +104,13 @@ public:
 	UPROPERTY(EditAnywhere, BlueprintReadWrite, Category = COMMON_MESHES_CATEGORY)
 	bool bBakeMeshes = true;
 
+	/** Specify how vertex colors should be imported */
+	UPROPERTY(EditAnywhere, BlueprintReadWrite, Category = COMMON_MESHES_CATEGORY)
+	TEnumAsByte<enum EInterchangeVertexColorImportOption> VertexColorImportOption;
+
+	/** Specify override color in the case that VertexColorImportOption is set to Override */
+	UPROPERTY(EditAnywhere, BlueprintReadWrite, Category = COMMON_MESHES_CATEGORY)
+	FColor VertexOverrideColor;
 	
 	//////	STATIC_MESHES_CATEGORY Properties //////
 
diff --git a/Engine/Source/Runtime/Engine/Private/Rendering/SkeletalMeshLODImporterData.cpp b/Engine/Source/Runtime/Engine/Private/Rendering/SkeletalMeshLODImporterData.cpp
--- a/Engine/Source/Runtime/Engine/Private/Rendering/SkeletalMeshLODImporterData.cpp
+++ b/Engine/Source/Runtime/Engine/Private/Rendering/SkeletalMeshLODImporterData.cpp
@@ -1839,6 +1839,10 @@ FSkeletalMeshImportData FSkeletalMeshImportData::CreateFromMeshDescription(const
 			Wedge.MatIndex = Face.MatIndex;
 			const bool bSRGB = false; //avoid linear to srgb conversion
 			Wedge.Color = FLinearColor(VertexInstanceColors[VertexInstanceID]).ToFColor(bSRGB);
+			if (Wedge.Color != FColor::White)
+			{
+				SkelMeshImportData.bHasVertexColors = true;
+			}
 			for (int32 UVChannelIndex = 0; UVChannelIndex < (int32)(SkelMeshImportData.NumTexCoords); ++UVChannelIndex)
 			{
 				Wedge.UVs[UVChannelIndex] = VertexInstanceUVs.Get(VertexInstanceID, UVChannelIndex);
