Fixed non-unity build for InterchangeGraphInspectorPipeline

diff --git a/Engine/Plugins/Editor/InterchangeEditor/Source/Pipelines/Private/InterchangeGraphInspectorPipeline.cpp b/Engine/Plugins/Editor/InterchangeEditor/Source/Pipelines/Private/InterchangeGraphInspectorPipeline.cpp
--- a/Engine/Plugins/Editor/InterchangeEditor/Source/Pipelines/Private/InterchangeGraphInspectorPipeline.cpp
+++ b/Engine/Plugins/Editor/InterchangeEditor/Source/Pipelines/Private/InterchangeGraphInspectorPipeline.cpp
@@ -2,6 +2,8 @@
 #include "InterchangeGraphInspectorPipeline.h"
 
 #include "CoreMinimal.h"
+
+#include "Framework/Application/SlateApplication.h"
 #include "InterchangeSourceData.h"
 #include "Interfaces/IMainFrameModule.h"
 #include "Misc/Paths.h"
@@ -9,6 +11,7 @@
 #include "SInterchangeGraphInspectorWindow.h"
 #include "UObject/Object.h"
 #include "UObject/ObjectMacros.h"
+#include "Widgets/SWindow.h"
 
 bool UInterchangeGraphInspectorPipeline::ExecutePreImportPipeline(UInterchangeBaseNodeContainer* BaseNodeContainer, const TArray<UInterchangeSourceData*>& SourceDatas)
 {
