Fixed non-unity build.

diff --git a/Engine/Plugins/Editor/InterchangeEditor/Source/ResultsBrowser/Private/SInterchangeResultsBrowserWindow.cpp b/Engine/Plugins/Editor/InterchangeEditor/Source/ResultsBrowser/Private/SInterchangeResultsBrowserWindow.cpp
--- a/Engine/Plugins/Editor/InterchangeEditor/Source/ResultsBrowser/Private/SInterchangeResultsBrowserWindow.cpp
+++ b/Engine/Plugins/Editor/InterchangeEditor/Source/ResultsBrowser/Private/SInterchangeResultsBrowserWindow.cpp
@@ -7,7 +7,6 @@
 #include "IDocumentation.h"
 #include "Modules/ModuleManager.h"
 #include "Styling/SlateIconFinder.h"
-#include "Widgets/Docking/SDockTab.h"
 #include "Widgets/Images/SImage.h"
 #include "Widgets/Input/SButton.h"
 #include "Widgets/Input/SCheckBox.h"
diff --git a/Engine/Plugins/Editor/InterchangeEditor/Source/ResultsBrowser/Private/SInterchangeResultsBrowserWindow.h b/Engine/Plugins/Editor/InterchangeEditor/Source/ResultsBrowser/Private/SInterchangeResultsBrowserWindow.h
--- a/Engine/Plugins/Editor/InterchangeEditor/Source/ResultsBrowser/Private/SInterchangeResultsBrowserWindow.h
+++ b/Engine/Plugins/Editor/InterchangeEditor/Source/ResultsBrowser/Private/SInterchangeResultsBrowserWindow.h
@@ -9,6 +9,7 @@
 #include "Styling/SlateColor.h"
 #include "Widgets/DeclarativeSyntaxSupport.h"
 #include "Widgets/SCompoundWidget.h"
+#include "Widgets/Docking/SDockTab.h"
 #include "Widgets/Views/SListView.h"
 #include "Widgets/Views/STableViewBase.h"
 #include "Widgets/Views/STableRow.h"
diff --git a/Engine/Plugins/Runtime/Interchange/Source/Nodes/Public/InterchangePhysicsAssetFactoryNode.h b/Engine/Plugins/Runtime/Interchange/Source/Nodes/Public/InterchangePhysicsAssetFactoryNode.h
--- a/Engine/Plugins/Runtime/Interchange/Source/Nodes/Public/InterchangePhysicsAssetFactoryNode.h
+++ b/Engine/Plugins/Runtime/Interchange/Source/Nodes/Public/InterchangePhysicsAssetFactoryNode.h
@@ -6,6 +6,7 @@
 #include "Nodes/InterchangeBaseNode.h"
 
 #if WITH_ENGINE
+#include "Animation/Skeleton.h"
 #include "PhysicsEngine/PhysicsAsset.h"
 #endif
 
